Network Prediction Insights
Unreal Insights based debugging tool for network prediction
enable tracing with -trace=NetworkPrediction

#RNX

diff --git a/Engine/Plugins/Runtime/NetworkPredictionInsights/NetworkPredictionInsights.uplugin b/Engine/Plugins/Runtime/NetworkPredictionInsights/NetworkPredictionInsights.uplugin
new file mode 100644
--- /dev/null
+++ b/Engine/Plugins/Runtime/NetworkPredictionInsights/NetworkPredictionInsights.uplugin
@@ -0,0 +1,28 @@
+{
+	"FileVersion": 3,
+	"Version": 1,
+	"VersionName": "1.0",
+	"FriendlyName": "Network Prediction Insights",
+	"Description": "Allows debugging of NetworkPrediction via Unreal Insights",
+	"Category": "Insights",
+	"CreatedBy": "Epic Games, Inc.",
+	"CreatedByURL": "http://epicgames.com",
+	"DocsURL": "",
+	"MarketplaceURL": "",
+	"SupportURL": "",
+	"EnabledByDefault": false,
+	"CanContainContent": false,
+	"IsBetaVersion": false,
+	"Installed": false,
+	"SupportedPrograms": [ "UnrealInsights" ],
+	"Modules": [
+		{
+			"Name": "NetworkPredictionInsights",
+			"Type": "EditorAndProgram",
+			"LoadingPhase": "Default",
+			"WhitelistPrograms": [
+				"UnrealInsights"
+			]
+		}
+	]
+}
diff --git a/Engine/Plugins/Runtime/NetworkPredictionInsights/Source/NetworkPredictionInsights/NetworkPredictionInsights.Build.cs b/Engine/Plugins/Runtime/NetworkPredictionInsights/Source/NetworkPredictionInsights/NetworkPredictionInsights.Build.cs
new file mode 100644
--- /dev/null
+++ b/Engine/Plugins/Runtime/NetworkPredictionInsights/Source/NetworkPredictionInsights/NetworkPredictionInsights.Build.cs
@@ -0,0 +1,47 @@
+// Copyright Epic Games, Inc. All Rights Reserved.
+
+namespace UnrealBuildTool.Rules
+{
+	public class NetworkPredictionInsights : ModuleRules
+	{
+		public NetworkPredictionInsights(ReadOnlyTargetRules Target) : base(Target)
+		{
+			PrivateDependencyModuleNames.AddRange(
+			new string[]
+			{
+				"Core",
+				"InputCore",
+				"SlateCore",
+				"Slate",
+				"TraceLog",
+				"TraceAnalysis",
+				"TraceServices",
+				"TraceInsights",
+				"AssetRegistry",
+				"ApplicationCore",
+				"EditorStyle",
+			});
+
+            if (Target.bCompileAgainstEngine)
+			{
+				PrivateDependencyModuleNames.AddRange(
+				new string[]
+				{
+					"Engine",
+					"CoreUObject",
+				});
+			}
+
+			if (Target.Type == TargetType.Editor)
+			{
+				PrivateDependencyModuleNames.AddRange(
+				new string[]
+				{
+					"UnrealEd",
+					"EditorWidgets",
+				});
+			}
+		}
+	}
+}
+
diff --git a/Engine/Plugins/Runtime/NetworkPredictionInsights/Source/NetworkPredictionInsights/Private/NetworkPredictionAnalyzer.cpp b/Engine/Plugins/Runtime/NetworkPredictionInsights/Source/NetworkPredictionInsights/Private/NetworkPredictionAnalyzer.cpp
new file mode 100644
--- /dev/null
+++ b/Engine/Plugins/Runtime/NetworkPredictionInsights/Source/NetworkPredictionInsights/Private/NetworkPredictionAnalyzer.cpp
@@ -0,0 +1,191 @@
+// Copyright Epic Games, Inc. All Rights Reserved.
+
+#include "NetworkPredictionAnalyzer.h"
+#include "NetworkPredictionProvider.h"
+
+
+FNetworkPredictionAnalyzer::FNetworkPredictionAnalyzer(Trace::IAnalysisSession& InSession, FNetworkPredictionProvider& InNetworkPredictionProvider)
+	: Session(InSession)
+	, NetworkPredictionProvider(InNetworkPredictionProvider)
+{
+}
+
+void FNetworkPredictionAnalyzer::OnAnalysisBegin(const FOnAnalysisContext& Context)
+{
+	auto& Builder = Context.InterfaceBuilder;
+	
+	Builder.RouteEvent(RouteId_GameInstanceRegister, "NetworkPrediction", "GameInstanceRegister");
+	Builder.RouteEvent(RouteId_SimulationCreated, "NetworkPrediction", "SimulationCreated");
+	Builder.RouteEvent(RouteId_SimulationNetRole, "NetworkPrediction", "SimulationNetRole");
+	Builder.RouteEvent(RouteId_SimulationNetGUID, "NetworkPrediction", "SimulationNetGUID");	
+	Builder.RouteEvent(RouteId_SimulationTick, "NetworkPrediction", "SimulationTick");
+	Builder.RouteEvent(RouteId_OOBStateMod, "NetworkPrediction", "OOBStateMod");
+	Builder.RouteEvent(RouteId_ProduceInput, "NetworkPrediction", "ProduceInput");
+	Builder.RouteEvent(RouteId_SynthInput, "NetworkPrediction", "SynthInput");
+	Builder.RouteEvent(RouteId_SimulationEOF, "NetworkPrediction", "SimulationEOF");
+	Builder.RouteEvent(RouteId_NetSerializeRecv, "NetworkPrediction", "NetSerializeRecv");
+	Builder.RouteEvent(RouteId_NetSerializeCommit, "NetworkPrediction", "NetSerializeCommit");
+	Builder.RouteEvent(RouteId_NetSerializeFault, "NetworkPrediction", "NetSerializeFault");
+	Builder.RouteEvent(RouteId_InputCmd, "NetworkPrediction", "InputCmd");
+	Builder.RouteEvent(RouteId_SyncState, "NetworkPrediction", "SyncState");
+	Builder.RouteEvent(RouteId_AuxState, "NetworkPrediction", "AuxState");
+	Builder.RouteEvent(RouteId_PieBegin, "NetworkPrediction", "PieBegin");
+}
+
+void FNetworkPredictionAnalyzer::OnAnalysisEnd()
+{
+}
+
+bool FNetworkPredictionAnalyzer::OnEvent(uint16 RouteId, const FOnEventContext& Context)
+{
+	Trace::FAnalysisSessionEditScope _(Session);
+	const auto& EventData = Context.EventData;
+
+	auto ParseUserState = [&EventData, this](ENP_UserState Type)
+	{
+		int32 SimulationId = EventData.GetValue<uint32>("SimulationId");			 
+		int32 Frame = EventData.GetValue<int32>("SimulationFrame");
+		uint64 EngineFrame = EventData.GetValue<uint64>("EngineFrameNumber");
+
+		NetworkPredictionProvider.WriteUserState(SimulationId, Frame, EngineFrame, Type, Session.StoreString(reinterpret_cast<const TCHAR*>(EventData.GetAttachment())));
+	};
+	
+	
+	bool bUnhandled = false;
+	switch (RouteId)
+	{
+		case RouteId_GameInstanceRegister:
+		{
+			// Temp, just set version here until we are encoding it as the first piece of data
+			NetworkPredictionProvider.SetNetworkPredictionTraceVersion(1);
+
+			// This is no longer doing anything since we are essentially all actor-role based now
+			// EventData.GetValue<uint32>("GameInstanceId"),
+			// EventData.GetValue<bool>("IsServer")
+			break;
+		}
+		case RouteId_SimulationCreated:
+		{
+			FSimulationData::FConst& ConstData = NetworkPredictionProvider.WriteSimulationCreated(EventData.GetValue<uint32>("SimulationId"));
+			ConstData.DebugName = FString(reinterpret_cast<const TCHAR*>(EventData.GetAttachment()));
+			ConstData.ID.NetGUID = EventData.GetValue<uint32>("NetGUID");
+			ConstData.GameInstanceId = EventData.GetValue<uint32>("GameInstanceId");
+			break;
+		}
+		case RouteId_SimulationNetRole:
+		{
+			NetworkPredictionProvider.WriteSimulationNetRole(	EventData.GetValue<uint32>("SimulationId"),
+																EventData.GetValue<uint64>("EngineFrameNumber"),
+																(ENP_NetRole)EventData.GetValue<uint8>("NetRole"));
+			break;
+		}
+
+		case RouteId_SimulationNetGUID:
+		{
+			NetworkPredictionProvider.WriteSimulationNetGUID(EventData.GetValue<uint32>("SimulationId"), EventData.GetValue<uint32>("NetGUID"));																
+			break;
+		}
+
+		case RouteId_SimulationTick:
+		{
+			FSimulationData::FTick TickData;
+			TickData.EngineFrame = EventData.GetValue<uint64>("EngineFrameNumber");
+			TickData.StartMS = EventData.GetValue<int32>("StartMS");
+			TickData.EndMS = EventData.GetValue<int32>("EndMS");
+			TickData.OutputFrame = EventData.GetValue<int32>("OutputFrame");
+
+			NetworkPredictionProvider.WriteSimulationTick(EventData.GetValue<uint32>("SimulationId"), MoveTemp(TickData));
+			break;
+		}
+		case RouteId_SimulationEOF:
+		{
+			FSimulationData::FEngineFrame& FrameData = NetworkPredictionProvider.WriteSimulationEOF(EventData.GetValue<uint32>("SimulationId"));
+			FrameData.EngineFrame = EventData.GetValue<uint64>("EngineFrameNumber");
+			FrameData.EngineFrameDeltaTime = EventData.GetValue<double>("EngineFrameDeltaTime");
+			FrameData.EngineCurrentTime = EventData.GetValue<double>("EngineCurrentTime");
+
+			FrameData.TotalProcessedSimTimeMS = EventData.GetValue<int32>("TotalProcessedMS");
+			FrameData.TotalAllowedSimTimeMS = EventData.GetValue<int32>("TotalAllowedMS");
+		
+			FrameData.LastSentKeyframe = EventData.GetValue<int32>("LastSentKeyframe");
+			FrameData.LastReceivedKeyframe = EventData.GetValue<int32>("LastReceivedKeyframe");
+			break;
+		}
+
+		case RouteId_NetSerializeRecv:
+		{
+			FSimulationData::FNetSerializeRecv NetRecv;
+			NetRecv.EngineFrame = EventData.GetValue<uint64>("EngineFrameNumber");
+			NetRecv.SimTimeMS = EventData.GetValue<uint64>("ReceivedSimTimeMS");
+			NetRecv.Frame = EventData.GetValue<uint64>("ReceivedFrame");
+			
+			NetworkPredictionProvider.WriteNetRecv(EventData.GetValue<uint32>("SimulationId"), MoveTemp(NetRecv));
+			break;
+		}
+		
+		case RouteId_NetSerializeCommit:
+		{
+			NetworkPredictionProvider.WriteNetCommit(EventData.GetValue<uint32>("SimulationId"));
+			break;
+		}
+
+		case RouteId_NetSerializeFault:
+		{
+			NetworkPredictionProvider.WriteNetFault(EventData.GetValue<uint32>("SimulationId"));
+			break;
+		}
+
+		case RouteId_OOBStateMod:
+		{
+			NetworkPredictionProvider.WriteOOBStateMod(EventData.GetValue<uint32>("SimulationId"));
+			break;
+		}
+
+		case RouteId_ProduceInput:
+		{
+			NetworkPredictionProvider.WriteProduceInput(EventData.GetValue<uint32>("SimulationId"));
+			break;
+		}
+		
+		case RouteId_SynthInput:
+		{
+			NetworkPredictionProvider.WriteSynthInput(EventData.GetValue<uint32>("SimulationId"));
+			break;
+		}
+
+		case RouteId_InputCmd:
+		{
+			ParseUserState(ENP_UserState::Input);
+			break;
+		}
+
+		case RouteId_SyncState:
+		{
+			ParseUserState(ENP_UserState::Sync);
+			break;
+		}
+
+		case RouteId_AuxState:
+		{
+			ParseUserState(ENP_UserState::Aux);
+			break;
+		}
+
+		case RouteId_PieBegin:
+		{
+			NetworkPredictionProvider.WritePIEStart();
+			break;
+		}
+
+		default:
+			bUnhandled = true;
+			break;
+	}
+
+	if (!bUnhandled)
+	{
+		NetworkPredictionProvider.IncrementDataCounter();
+	}
+
+	return true;
+}
\ No newline at end of file
diff --git a/Engine/Plugins/Runtime/NetworkPredictionInsights/Source/NetworkPredictionInsights/Private/NetworkPredictionAnalyzer.h b/Engine/Plugins/Runtime/NetworkPredictionInsights/Source/NetworkPredictionInsights/Private/NetworkPredictionAnalyzer.h
new file mode 100644
--- /dev/null
+++ b/Engine/Plugins/Runtime/NetworkPredictionInsights/Source/NetworkPredictionInsights/Private/NetworkPredictionAnalyzer.h
@@ -0,0 +1,53 @@
+// Copyright Epic Games, Inc. All Rights Reserved.
+
+#pragma once
+
+#include "Trace/Analyzer.h"
+#include "TraceServices/Model/AnalysisSession.h"
+
+class FNetworkPredictionProvider;
+namespace Trace { class IAnalysisSession; }
+
+// Analyzes events that are contained in a trace,
+// Works by subscribing to events by name along with user-provided "route" identifiers
+// To analyze a trace, concrete IAnalyzer-derived objects are registered with a
+// FAnalysisContext which is then asked to launch and coordinate the analysis.
+//
+// The analyzer is what populates the data in the Provider.
+
+class FNetworkPredictionAnalyzer : public Trace::IAnalyzer
+{
+public:
+
+	FNetworkPredictionAnalyzer(Trace::IAnalysisSession& InSession, FNetworkPredictionProvider& InNetworkPredictionProvider);
+
+	virtual void OnAnalysisBegin(const FOnAnalysisContext& Context) override;
+	virtual void OnAnalysisEnd() override;
+	virtual bool OnEvent(uint16 RouteId, const FOnEventContext& Context) override;
+
+private:
+
+	enum : uint16
+	{
+		RouteId_GameInstanceRegister,
+		RouteId_SimulationCreated,
+		RouteId_SimulationNetRole,
+		RouteId_SimulationNetGUID,
+		RouteId_SimulationTick,
+		RouteId_OOBStateMod,
+		RouteId_ProduceInput,
+		RouteId_SynthInput,
+		RouteId_SimulationEOF,
+		RouteId_NetSerializeRecv,
+		RouteId_NetSerializeCommit,
+		RouteId_NetSerializeFault,
+		RouteId_InputCmd,
+		RouteId_SyncState,
+		RouteId_AuxState,
+		RouteId_PieBegin
+	};
+
+
+	Trace::IAnalysisSession& Session;
+	FNetworkPredictionProvider& NetworkPredictionProvider;
+};
\ No newline at end of file
diff --git a/Engine/Plugins/Runtime/NetworkPredictionInsights/Source/NetworkPredictionInsights/Private/NetworkPredictionInsightsModule.cpp b/Engine/Plugins/Runtime/NetworkPredictionInsights/Source/NetworkPredictionInsights/Private/NetworkPredictionInsightsModule.cpp
new file mode 100644
--- /dev/null
+++ b/Engine/Plugins/Runtime/NetworkPredictionInsights/Source/NetworkPredictionInsights/Private/NetworkPredictionInsightsModule.cpp
@@ -0,0 +1,139 @@
+// Copyright Epic Games, Inc. All Rights Reserved.
+
+#include "NetworkPredictionInsightsModule.h"
+#include "Features/IModularFeatures.h"
+#include "Insights/ITimingViewExtender.h"
+#include "Containers/Ticker.h"
+#include "Modules/ModuleManager.h"
+#include "Framework/Docking/LayoutExtender.h"
+#include "Insights/IUnrealInsightsModule.h"
+#include "HAL/PlatformApplicationMisc.h"
+#include "WorkspaceMenuStructure.h"
+#include "WorkspaceMenuStructureModule.h"
+#include "TraceServices/ITraceServicesModule.h"
+#include "Widgets/Docking/SDockTab.h"
+#include "Trace/StoreService.h"
+#include "Trace/StoreClient.h"
+
+#include "UI/SNPWindow.h"
+#include "UI/NetworkPredictionInsightsManager.h"
+#include "EditorStyleSet.h"
+
+#if WITH_ENGINE
+#include "Engine/Engine.h"
+#endif
+
+
+#define LOCTEXT_NAMESPACE "NetworkPredictionInsightsModule"
+
+const FName NetworkPredictionInsightsTabs::DocumentTab("DocumentTab");
+
+const FName FNetworkPredictionInsightsModule::InsightsTabName("NetworkPrediction");
+
+void FNetworkPredictionInsightsModule::StartupModule()
+{
+	IModularFeatures::Get().RegisterModularFeature(Trace::ModuleFeatureName, &NetworkPredictionTraceModule);
+
+	FNetworkPredictionInsightsManager::Initialize();
+	IUnrealInsightsModule& UnrealInsightsModule = FModuleManager::LoadModuleChecked<IUnrealInsightsModule>("TraceInsights");
+
+	// Auto spawn the Network Prediction Insights tab if we detect NP data
+	// Only do this in Standalone UnrealInsights.exe. In Editor the user will select the NPI window manually.
+	// (There is currently not way to extend the high level layout of unreal insights, e.g, the layout created in FTraceInsightsModule::CreateSessionBrowser)
+	// (only SetUnrealInsightsLayoutIni which is extending the layouts of pre made individual tabs, not the the overall session layout)
+#if !(WITH_EDITOR)
+	TickerHandle = FTicker::GetCoreTicker().AddTicker(TEXT("NetworkPredictionInsights"), 0.0f, [&UnrealInsightsModule](float DeltaTime)
+	{
+		auto SessionPtr = UnrealInsightsModule.GetAnalysisSession();
+		if (SessionPtr.IsValid())
+		{
+			Trace::FAnalysisSessionReadScope SessionReadScope(*SessionPtr.Get());
+			if (const INetworkPredictionProvider* NetworkPredictionProvider = ReadNetworkPredictionProvider(*SessionPtr.Get()))
+			{
+				auto NetworkPredictionTraceVersion = NetworkPredictionProvider->GetNetworkPredictionTraceVersion();
+
+				if (NetworkPredictionTraceVersion > 0)
+				{
+					static bool HasSpawnedTab = false;
+					if (!HasSpawnedTab && FGlobalTabmanager::Get()->HasTabSpawner(FNetworkPredictionInsightsModule::InsightsTabName))
+					{
+						HasSpawnedTab = true;
+						FGlobalTabmanager::Get()->TryInvokeTab(FNetworkPredictionInsightsModule::InsightsTabName);
+					}
+				}
+			}
+		}
+
+		return true;
+	});
+#endif
+
+	// Actually register our tab spawner
+	FTabSpawnerEntry& TabSpawnerEntry = FGlobalTabmanager::Get()->RegisterNomadTabSpawner(FNetworkPredictionInsightsModule::InsightsTabName,
+		FOnSpawnTab::CreateLambda([](const FSpawnTabArgs& Args)
+	{
+		const TSharedRef<SDockTab> DockTab = SNew(SDockTab)
+			.TabRole(ETabRole::NomadTab);
+
+		TSharedRef<SNPWindow> Window = SNew(SNPWindow, DockTab, Args.GetOwnerWindow());
+		DockTab->SetContent(Window);
+		return DockTab;
+	}))
+		.SetDisplayName(NSLOCTEXT("FNetworkPredictionInsightsModule", "NetworkPredictionTabTitle", "Network Prediction Insights"))
+		.SetTooltipText(NSLOCTEXT("FNetworkPredictionInsightsModule", "FilteringTabTooltip", "Opens the Network Prediction Insights tab."))
+		.SetIcon(FSlateIcon(FEditorStyle::GetStyleSetName(), "ProfilerCommand.StatsProfiler.Small"));
+
+	//TabSpawnerEntry.SetGroup(WorkspaceMenu::GetMenuStructure().GetDeveloperToolsProfilingCategory());
+	TabSpawnerEntry.SetGroup(WorkspaceMenu::GetMenuStructure().GetToolsCategory());
+
+	// -------------------------------------------------------------
+
+#if WITH_EDITOR
+	if (!IsRunningCommandlet())
+	{
+		// Conditionally create local store service after engine init (if someone doesn't beat us to it).
+		// This is temp until a more formal local server is done by the insights system.
+		StoreServiceHandle = FCoreDelegates::OnFEngineLoopInitComplete.AddLambda([this]
+		{
+			IUnrealInsightsModule& UnrealInsightsModule = FModuleManager::LoadModuleChecked<IUnrealInsightsModule>("TraceInsights");
+			if (!UnrealInsightsModule.GetStoreClient())
+			{
+				// Create the Store Service.
+				FString StoreDir = FPaths::ProjectSavedDir() / TEXT("TraceSessions");
+				Trace::FStoreService::FDesc StoreServiceDesc;
+				StoreServiceDesc.StoreDir = *StoreDir;
+				StoreServiceDesc.RecorderPort = 0; // Let system decide port
+				StoreServiceDesc.ThreadCount = 2;
+				StoreService = TSharedPtr<Trace::FStoreService>(Trace::FStoreService::Create(StoreServiceDesc));
+
+				FCoreDelegates::OnPreExit.AddLambda([this]() {
+					StoreService.Reset();
+				});
+
+				// Connect to our newly created store and setup the insights module
+				ensure(UnrealInsightsModule.ConnectToStore(TEXT("localhost"), StoreService->GetPort()));
+				Trace::SendTo(TEXT("localhost"), StoreService->GetRecorderPort());
+
+				UnrealInsightsModule.CreateSessionViewer(false);
+				UnrealInsightsModule.StartAnalysisForLastLiveSession();
+			}
+		});
+	}
+#endif
+}
+
+void FNetworkPredictionInsightsModule::ShutdownModule()
+{
+	if (StoreServiceHandle.IsValid())
+	{
+		FCoreDelegates::OnFEngineLoopInitComplete.Remove(StoreServiceHandle);
+	}
+
+	FTicker::GetCoreTicker().RemoveTicker(TickerHandle);
+
+	IModularFeatures::Get().UnregisterModularFeature(Trace::ModuleFeatureName, &NetworkPredictionTraceModule);
+}
+
+IMPLEMENT_MODULE(FNetworkPredictionInsightsModule, NetworkPredictionInsights);
+
+#undef LOCTEXT_NAMESPACE
diff --git a/Engine/Plugins/Runtime/NetworkPredictionInsights/Source/NetworkPredictionInsights/Private/NetworkPredictionInsightsModule.h b/Engine/Plugins/Runtime/NetworkPredictionInsights/Source/NetworkPredictionInsights/Private/NetworkPredictionInsightsModule.h
new file mode 100644
--- /dev/null
+++ b/Engine/Plugins/Runtime/NetworkPredictionInsights/Source/NetworkPredictionInsights/Private/NetworkPredictionInsightsModule.h
@@ -0,0 +1,36 @@
+// Copyright Epic Games, Inc. All Rights Reserved.
+
+#pragma once
+
+#include "CoreMinimal.h"
+#include "Modules/ModuleInterface.h"
+#include "NetworkPredictionTraceModule.h"
+#include "Framework/Docking/TabManager.h"
+#include "INetworkPredictionInsightsModule.h"
+
+struct FInsightsMajorTabExtender;
+
+namespace Trace 
+{
+	class FStoreService;
+	class FStoreClient;
+}
+
+class FNetworkPredictionInsightsModule : public IModuleInterface
+{
+public:
+	// IModuleInterface interface
+	virtual void StartupModule() override;
+	virtual void ShutdownModule() override;
+
+private:
+	FNetworkPredictionTraceModule NetworkPredictionTraceModule;
+
+	FDelegateHandle TickerHandle;
+	FDelegateHandle StoreServiceHandle;
+
+	TSharedPtr<Trace::FStoreService> StoreService;
+	//TWeakPtr<FTabManager> WeakTimingProfilerTabManager;
+
+	static const FName InsightsTabName;	
+};
diff --git a/Engine/Plugins/Runtime/NetworkPredictionInsights/Source/NetworkPredictionInsights/Private/NetworkPredictionInsightsStyle.h b/Engine/Plugins/Runtime/NetworkPredictionInsights/Source/NetworkPredictionInsights/Private/NetworkPredictionInsightsStyle.h
new file mode 100644
--- /dev/null
+++ b/Engine/Plugins/Runtime/NetworkPredictionInsights/Source/NetworkPredictionInsights/Private/NetworkPredictionInsightsStyle.h
@@ -0,0 +1,91 @@
+// Copyright Epic Games, Inc. All Rights Reserved.
+
+#pragma once
+
+#include "Styling/SlateStyle.h"
+#include "Styling/SlateStyleRegistry.h"
+
+#define IMAGE_BRUSH(RelativePath, ...) FSlateImageBrush(RootToContentDir(RelativePath, TEXT(".png")), __VA_ARGS__)
+#define BOX_BRUSH(RelativePath, ...) FSlateBoxBrush(RootToContentDir(RelativePath, TEXT(".png")), __VA_ARGS__)
+#define BORDER_BRUSH(RelativePath, ...) FSlateBorderBrush(RootToContentDir(RelativePath, TEXT(".png")), __VA_ARGS__)
+#define TTF_FONT(RelativePath, ...) FSlateFontInfo(RootToContentDir(RelativePath, TEXT(".ttf")), __VA_ARGS__)
+#define OTF_FONT(RelativePath, ...) FSlateFontInfo(RootToContentDir(RelativePath, TEXT(".otf")), __VA_ARGS__)
+
+class FNetworkPredictionInsightsStyle
+	: public FSlateStyleSet
+{
+public:
+	FNetworkPredictionInsightsStyle()
+		: FSlateStyleSet("NetworkPredictionInsightsStyle")
+		, SelectionColor_LinearRef(MakeShareable(new FLinearColor( 0.728f, 0.364f, 0.003f)))
+		, SelectionColor_Pressed_LinearRef(MakeShareable(new FLinearColor( 0.701f, 0.225f, 0.003f)))
+		, SelectionColor(SelectionColor_LinearRef)
+		, SelectionColor_Pressed(SelectionColor_Pressed_LinearRef)
+	{
+		const FVector2D Icon10x10(10.0f, 10.0f);
+		const FVector2D Icon16x16(16.0f, 16.0f);
+		const FVector2D Icon20x20(20.0f, 20.0f);
+		const FVector2D Icon24x24(24.0f, 24.0f);
+		const FVector2D Icon32x32(32.0f, 32.0f);
+		const FVector2D Icon40x40(40.0f, 40.0f);
+		SetContentRoot(FPaths::EnginePluginsDir() / TEXT("Animation/NetworkPredictionInsights/Content"));
+
+		const FSlateColor DefaultForeground(FLinearColor(0.72f, 0.72f, 0.72f, 1.f));
+
+		Set("SchematicViewRootLeft", new BOX_BRUSH("SchematicViewRootLeft", FMargin(4/16.0f, 8.0f/16.0f, 4/16.0f, 4/16.0f)));
+		Set("SchematicViewRootMid", new BOX_BRUSH("SchematicViewRootMid", FMargin(4/16.0f, 8.0f/16.0f, 4/16.0f, 4/16.0f)));
+
+		Set("SchematicViewViewButton", FButtonStyle()
+			.SetNormal(FSlateNoResource())
+			.SetHovered(BOX_BRUSH("RoundedSelection_16x", 4.0f/16.0f, SelectionColor))
+			.SetPressed(BOX_BRUSH("RoundedSelection_16x", 4.0f/16.0f, SelectionColor_Pressed))
+			.SetNormalPadding(FMargin(2, 2, 2, 2))
+			.SetPressedPadding(FMargin(2, 3, 2, 1))
+		);
+
+		Set("SchematicViewViewButtonIcon", new IMAGE_BRUSH("view_button", Icon20x20));
+
+		Set("TransportControls.HyperlinkSpinBox", FSpinBoxStyle()
+			.SetTextPadding(FMargin(0))
+			.SetBackgroundBrush(BORDER_BRUSH("HyperlinkDotted", FMargin(0, 0, 0, 3 / 16.0f), FSlateColor::UseSubduedForeground()))
+			.SetHoveredBackgroundBrush(FSlateNoResource())
+			.SetInactiveFillBrush(FSlateNoResource())
+			.SetActiveFillBrush(FSlateNoResource())
+			.SetForegroundColor(FSlateColor::UseSubduedForeground())
+			.SetArrowsImage(FSlateNoResource())
+		);
+
+		FSlateStyleRegistry::RegisterSlateStyle(*this);
+	}
+
+	static FNetworkPredictionInsightsStyle& Get()
+	{
+		static FNetworkPredictionInsightsStyle Inst;
+		return Inst;
+	}
+	
+	~FNetworkPredictionInsightsStyle()
+	{
+		FSlateStyleRegistry::UnRegisterSlateStyle(*this);
+	}
+
+	const TSharedRef<FLinearColor> SelectionColor_LinearRef;
+	const TSharedRef<FLinearColor> SelectionColor_Pressed_LinearRef;
+	const FSlateColor SelectionColor;
+	const FSlateColor SelectionColor_Pressed;
+
+	const TSharedRef<FLinearColor> SimFrameColor_Confirmed_Ref;
+	const TSharedRef<FLinearColor> SimFrameColor_Corrected_Ref;
+	const TSharedRef<FLinearColor> SimFrameColor_Predicted_Ref;
+	const FSlateColor SimFrameColor_Confirmed;
+	const FSlateColor SimFrameColor_Confirmed;
+	const FSlateColor SimFrameColor_Confirmed;
+
+	0.53333  0.67059  0.53333
+};
+
+#undef IMAGE_BRUSH
+#undef BOX_BRUSH
+#undef BORDER_BRUSH
+#undef TTF_FONT
+#undef OTF_FONT
diff --git a/Engine/Plugins/Runtime/NetworkPredictionInsights/Source/NetworkPredictionInsights/Private/NetworkPredictionProvider.cpp b/Engine/Plugins/Runtime/NetworkPredictionInsights/Source/NetworkPredictionInsights/Private/NetworkPredictionProvider.cpp
new file mode 100644
--- /dev/null
+++ b/Engine/Plugins/Runtime/NetworkPredictionInsights/Source/NetworkPredictionInsights/Private/NetworkPredictionProvider.cpp
@@ -0,0 +1,397 @@
+// Copyright Epic Games, Inc. All Rights Reserved.
+
+#include "NetworkPredictionProvider.h"
+
+FName FNetworkPredictionProvider::ProviderName("NetworkPredictionProvider");
+
+// -----------------------------------------------------------------------------
+
+FNetworkPredictionProvider::FNetworkPredictionProvider(Trace::IAnalysisSession& InSession)
+	: Session(InSession)
+{
+
+}
+
+void FNetworkPredictionProvider::SetNetworkPredictionTraceVersion(uint32 Version)
+{
+	Session.WriteAccessCheck();
+	NetworkPredictionTraceVersion = Version;
+}
+
+// -----------------------------------------------------------------------------
+
+FSimulationData::FConst& FNetworkPredictionProvider::WriteSimulationCreated(uint32 SimulationId)
+{
+	FSimulationData::FConst& SimulationConstData = FindOrAdd(SimulationId)->ConstData;
+	SimulationConstData.ID.PIESession = PIESessionCounter;
+	return SimulationConstData;
+}
+
+void FNetworkPredictionProvider::WriteSimulationTick(uint32 SimulationId, FSimulationData::FTick&& InTick)
+{
+	FSimulationData& SimulationData = FindOrAdd(SimulationId).Get();
+
+	SimulationData.Analysis.PendingUserStateSource = ENP_UserStateSource::SimTick;
+	SimulationData.Analysis.PendingCommitUserStates.Reset();
+
+	// ---------------------------------------------------------------------------------------------
+	//	New tick may cause previous ticks to be trashed.
+	//	FIXME: Dumb approach: loop through whole list. Can be cached off 
+	//	note: don't do this after adding new tick to list since we'll end up trashing ourselves
+	// ---------------------------------------------------------------------------------------------
+	if (SimulationData.Ticks.Num() > 0)
+	{
+		for (auto It = SimulationData.Ticks.GetIterator(); It; ++It)
+		{
+			FSimulationData::FTick& Tick = const_cast<FSimulationData::FTick&>(*It);
+			if (InTick.StartMS < Tick.EndMS && Tick.TrashedEngineFrame == 0)
+			{
+				Tick.TrashedEngineFrame = InTick.EngineFrame;
+			}
+		}
+	}
+
+	// Add it to the list
+	FSimulationData::FTick& NewTick = SimulationData.Ticks.PushBack();
+	NewTick = MoveTemp(InTick);
+
+	// Repredict if we've already simulated past this time before
+	if (NewTick.StartMS < SimulationData.Analysis.MaxTickSimTimeMS)
+	{
+		NewTick.bRepredict = true;
+	}
+	SimulationData.Analysis.MaxTickSimTimeMS = FMath::Max(SimulationData.Analysis.MaxTickSimTimeMS, NewTick.EndMS);
+
+	// ---------------------------------------------------------------------------------------------
+	//	Link this new tick with pending NetRecvs if possible
+	// ---------------------------------------------------------------------------------------------
+
+	TArray<FSimulationData::FNetSerializeRecv*> WorkingPendingNetSerializeRecv = MoveTemp(SimulationData.Analysis.PendingNetSerializeRecv);
+	for (FSimulationData::FNetSerializeRecv* PendingRecv : WorkingPendingNetSerializeRecv)
+	{
+		check(PendingRecv);
+
+		const FSimTime RecvSimTimeMS = PendingRecv->SimTimeMS;
+
+		const bool bIsAuthority = (SimulationData.SparseData.Read(PendingRecv->EngineFrame)->NetRole == ENP_NetRole::Authority);
+
+		// On non authority, this becomes a confirmed frame now if it wasn't already marked
+		if (!bIsAuthority)
+		{
+			if (PendingRecv->Status == ENetSerializeRecvStatus::Unknown)
+			{
+				PendingRecv->Status = ENetSerializeRecvStatus::Confirm;
+			}
+		}
+
+		// Perfect Match
+		if (NewTick.StartMS == RecvSimTimeMS)
+		{
+			NewTick.StartNetRecv = PendingRecv;
+			PendingRecv->NextTick = &NewTick;
+		}
+		// Ahead of last NetRecv
+		else if (NewTick.StartMS > RecvSimTimeMS)
+		{
+			// Instead of attaching to the Orphaned list, find a tick to attach it to. This is a bit dicey and may not be
+			// a good solution if we want to treat recvs that came in mid local predicted frame as confirmed (e.g, for sim proxies)
+			auto FindNetRecvAHome = [&]()
+			{
+				for (auto It = SimulationData.Ticks.GetIteratorFromItem(SimulationData.Ticks.Num()-1); It; --It)
+				{
+					FSimulationData::FTick& PrevTick = const_cast<FSimulationData::FTick&>(*It);
+
+					if (!PrevTick.StartNetRecv && PrevTick.StartMS == PendingRecv->SimTimeMS)
+					{
+						PrevTick.StartNetRecv = PendingRecv;
+						PendingRecv->NextTick = &PrevTick;
+						return true;
+					}
+				}
+				return false;
+			};
+
+			if (!FindNetRecvAHome())
+			{
+				// No home for this recv so put him on the orphan list so he still gets drawn
+				if (PendingRecv->Status == ENetSerializeRecvStatus::Unknown)
+				{
+					PendingRecv->Status = ENetSerializeRecvStatus::Stale;
+				}
+			}
+		}
+		else
+		{
+			SimulationData.Analysis.PendingNetSerializeRecv.Add(PendingRecv);
+		}
+	}
+}
+
+FSimulationData::FEngineFrame& FNetworkPredictionProvider::WriteSimulationEOF(uint32 SimulationId)
+{
+	return FindOrAdd(SimulationId)->EOFState.PushBack();
+}
+
+void FNetworkPredictionProvider::WriteNetRecv(uint32 SimulationId, FSimulationData::FNetSerializeRecv&& InNetRecv)
+{
+	FSimulationData& SimulationData = FindOrAdd(SimulationId).Get();
+
+	SimulationData.Analysis.PendingUserStateSource = ENP_UserStateSource::NetRecv;
+	SimulationData.Analysis.PendingCommitUserStates.Reset();
+
+	FSimulationData::FNetSerializeRecv &NewNetSerializeRecv = SimulationData.NetRecv.PushBack();
+	NewNetSerializeRecv = MoveTemp(InNetRecv);
+
+	if (SimulationData.SparseData.Read(NewNetSerializeRecv.EngineFrame)->NetRole == ENP_NetRole::Authority)
+	{
+		// Authority gets net receives before running ticks, so we accumulate all of them
+		SimulationData.Analysis.PendingNetSerializeRecv.Add(&NewNetSerializeRecv);
+	}
+	else
+	{
+		if (SimulationData.Analysis.PendingNetSerializeRecv.Num() > 0)
+		{
+			ensure(SimulationData.Analysis.PendingNetSerializeRecv.Num() == 1);
+
+			if (SimulationData.Analysis.PendingNetSerializeRecv.Last()->Status == ENetSerializeRecvStatus::Unknown)
+			{
+				SimulationData.Analysis.PendingNetSerializeRecv.Last()->Status = ENetSerializeRecvStatus::Stale;
+			}
+			SimulationData.Analysis.PendingNetSerializeRecv.Reset();
+		}
+		SimulationData.Analysis.PendingNetSerializeRecv.Add(&NewNetSerializeRecv);
+
+		if (NewNetSerializeRecv.SimTimeMS > SimulationData.Analysis.MaxTickSimTimeMS)
+		{
+			NewNetSerializeRecv.Status = ENetSerializeRecvStatus::Jump;
+		}
+	}
+
+	// Mark ConfirmedEngineFrame on pending frame
+	// We want Tick.ConfirmedEngineFrame to be the earliest engine frame we received that was timestamped > what the tick was
+	// NetRecvs will always increase in sim time, but our ticks do not
+	// So we need to search the remainder of the list, knowing that tick time will be not be linear
+	if (SimulationData.Ticks.Num() > SimulationData.Analysis.NetRecvItemIdx)
+	{
+		bool bIncrementSavedIdx = true;
+		for (auto It = SimulationData.Ticks.GetIteratorFromItem(SimulationData.Analysis.NetRecvItemIdx); It; ++It)
+		{
+			FSimulationData::FTick& Tick = const_cast<FSimulationData::FTick&>(*It);
+			if (Tick.EndMS > NewNetSerializeRecv.SimTimeMS)
+			{
+				// Once we encounter a frame that is still unconfirmed, we can no longer increment our saved position,
+				// We need to pick up back here when we get a newer NetRecv.
+				bIncrementSavedIdx = false;
+			}
+			else if (Tick.ConfirmedEngineFrame == 0)
+			{
+				// This is the first NetRecv that this frame was confirmed on
+				Tick.ConfirmedEngineFrame = NewNetSerializeRecv.EngineFrame;
+			}
+
+			if (bIncrementSavedIdx)
+			{
+				SimulationData.Analysis.NetRecvItemIdx++;
+			}
+		}
+	}
+}
+
+void FNetworkPredictionProvider::WriteNetCommit(uint32 SimulationId)
+{
+	FSimulationData& SimulationData = FindOrAdd(SimulationId).Get();
+
+	// Mark pending user states committed
+	for (FSimulationData::FUserState* UserState : SimulationData.Analysis.PendingCommitUserStates)
+	{
+		if (ensure(UserState->Source == ENP_UserStateSource::NetRecv))
+		{
+			UserState->Source = ENP_UserStateSource::NetRecvCommit;
+		}
+	}
+	SimulationData.Analysis.PendingCommitUserStates.Reset();
+
+	// Pending NetReceives also become committed here
+	for (FSimulationData::FNetSerializeRecv* PendingRecv : SimulationData.Analysis.PendingNetSerializeRecv)
+	{
+		check(PendingRecv);
+
+		const FSimTime RecvSimTimeMS = PendingRecv->SimTimeMS;
+		const bool bIsAuthority = (SimulationData.SparseData.Read(PendingRecv->EngineFrame)->NetRole == ENP_NetRole::Authority);
+
+		if (PendingRecv->Status == ENetSerializeRecvStatus::Unknown)
+		{
+			if (bIsAuthority)
+			{
+				// Authority this is "confirmed" (misnomer): we wait for net recvs to advance the sim (unless synth)
+				PendingRecv->Status = ENetSerializeRecvStatus::Confirm;
+			}
+			else
+			{
+				// Non authority, this caused a rollback/correction since we committed to our buffers
+				PendingRecv->Status = ENetSerializeRecvStatus::Rollback;
+			}
+		}
+	}
+}
+
+void FNetworkPredictionProvider::WriteNetFault(uint32 SimulationId)
+{
+	FSimulationData& SimulationData = FindOrAdd(SimulationId).Get();
+	
+	for (FSimulationData::FNetSerializeRecv* PendingRecv : SimulationData.Analysis.PendingNetSerializeRecv)
+	{
+		ensure(PendingRecv->Status == ENetSerializeRecvStatus::Unknown);
+		PendingRecv->Status = ENetSerializeRecvStatus::Fault;
+	}
+}
+
+void FNetworkPredictionProvider::WriteOOBStateMod(uint32 SimulationId)
+{
+	FSimulationData& SimulationData = FindOrAdd(SimulationId).Get();
+	SimulationData.Analysis.PendingCommitUserStates.Reset();
+	SimulationData.Analysis.PendingUserStateSource = ENP_UserStateSource::OOB;
+}
+
+void FNetworkPredictionProvider::WriteProduceInput(uint32 SimulationId)
+{
+	FSimulationData& SimulationData = FindOrAdd(SimulationId).Get();
+	SimulationData.Analysis.PendingCommitUserStates.Reset();
+	SimulationData.Analysis.PendingUserStateSource = ENP_UserStateSource::ProduceInput;
+}
+
+void FNetworkPredictionProvider::WriteSynthInput(uint32 SimulationId)
+{
+	FSimulationData& SimulationData = FindOrAdd(SimulationId).Get();
+	SimulationData.Analysis.PendingCommitUserStates.Reset();
+	SimulationData.Analysis.PendingUserStateSource = ENP_UserStateSource::SynthInput;
+}
+
+void FNetworkPredictionProvider::WriteSimulationNetRole(uint32 SimulationId, uint64 EngineFrame, ENP_NetRole Role)
+{
+	FindOrAdd(SimulationId)->SparseData.Write(EngineFrame)->NetRole = Role;
+}
+
+void FNetworkPredictionProvider::WriteSimulationNetGUID(uint32 SimulationId, uint32 NetGUID)
+{
+	FindOrAdd(SimulationId)->ConstData.ID.NetGUID = NetGUID;
+}
+
+void FNetworkPredictionProvider::WriteUserState(uint32 SimulationId, int32 Frame, uint64 EngineFrame, ENP_UserState Type, const TCHAR* UserStr)
+{
+	FSimulationData& SimulationData = FindOrAdd(SimulationId).Get();
+	ensure(SimulationData.Analysis.PendingUserStateSource != ENP_UserStateSource::Unknown);
+
+	FSimulationData::FUserState& NewUserState = SimulationData.UserData.Store[(int32)Type].Push(Frame, EngineFrame);
+	NewUserState.UserStr = UserStr;
+	NewUserState.Source = SimulationData.Analysis.PendingUserStateSource;
+
+	if (SimulationData.Analysis.PendingUserStateSource == ENP_UserStateSource::NetRecv)
+	{
+		SimulationData.Analysis.PendingCommitUserStates.Add(&NewUserState);
+	}
+}
+
+void FNetworkPredictionProvider::WritePIEStart()
+{
+	PIESessionCounter++;
+}
+
+TSharedRef<FSimulationData>& FNetworkPredictionProvider::FindOrAdd(uint32 SimulationId)
+{
+	for (TSharedRef<FSimulationData>& Data : ProviderData)
+	{
+		if (Data->SimulationId == SimulationId)
+		{
+			return Data;
+		}
+	}
+
+	ProviderData.Add(MakeShareable(new FSimulationData(SimulationId, Session.GetLinearAllocator())));
+	return ProviderData.Last();
+}
+
+// -----------------------------------------------------------------------------
+//
+// -----------------------------------------------------------------------------
+
+const INetworkPredictionProvider* ReadNetworkPredictionProvider(const Trace::IAnalysisSession& Session)
+{
+	Session.ReadAccessCheck();
+	return Session.ReadProvider<INetworkPredictionProvider>(FNetworkPredictionProvider::ProviderName);
+}
+
+const TCHAR* LexToString(ENP_NetRole Role)
+{
+	switch (Role)
+	{
+	case ENP_NetRole::None:
+		return TEXT("None");
+	case ENP_NetRole::SimulatedProxy:
+		return TEXT("Simulated");
+	case ENP_NetRole::AutonomousProxy:
+		return TEXT("Autonomous");
+	case ENP_NetRole::Authority:
+		return TEXT("Authority");	
+	default:
+		return TEXT("Unknown");
+	}
+}
+
+const TCHAR* LexToString(ENP_UserState State)
+{
+	switch(State)
+	{
+	case ENP_UserState::Input:
+		return TEXT("Input");
+	case ENP_UserState::Sync:
+		return TEXT("Sync");
+	case ENP_UserState::Aux:
+		return TEXT("Aux");
+	default:
+		return TEXT("Unknown");
+	}
+}
+const TCHAR* LexToString(ENP_UserStateSource Source)
+{
+	switch(Source)
+	{
+	case ENP_UserStateSource::ProduceInput:
+		return TEXT("ProduceInput");
+	case ENP_UserStateSource::SynthInput:
+		return TEXT("SynthInput");
+	case ENP_UserStateSource::SimTick:
+		return TEXT("SimTick");
+	case ENP_UserStateSource::NetRecv:
+		return TEXT("NetRecv (Uncommited)");
+	case ENP_UserStateSource::NetRecvCommit:
+		return TEXT("NetRecv Committed");
+	case ENP_UserStateSource::OOB:
+		return TEXT("OOB");
+	default:
+		return TEXT("Unknown");
+	}
+}
+
+const TCHAR* LexToString(ENetSerializeRecvStatus Status)
+{
+	switch(Status)
+	{
+	case ENetSerializeRecvStatus::Unknown:
+		return TEXT("Unknown");
+	case ENetSerializeRecvStatus::Confirm:
+		return TEXT("Confirm");
+	case ENetSerializeRecvStatus::Rollback:
+		return TEXT("Rollback");
+	case ENetSerializeRecvStatus::Jump:
+		return TEXT("Jump");
+	case ENetSerializeRecvStatus::Fault:
+		return TEXT("Fault");
+	case ENetSerializeRecvStatus::Stale:
+		return TEXT("Stale");
+
+	default:
+		ensure(false);
+		return TEXT("???");
+	}
+}
\ No newline at end of file
diff --git a/Engine/Plugins/Runtime/NetworkPredictionInsights/Source/NetworkPredictionInsights/Private/NetworkPredictionProvider.h b/Engine/Plugins/Runtime/NetworkPredictionInsights/Source/NetworkPredictionInsights/Private/NetworkPredictionProvider.h
new file mode 100644
--- /dev/null
+++ b/Engine/Plugins/Runtime/NetworkPredictionInsights/Source/NetworkPredictionInsights/Private/NetworkPredictionProvider.h
@@ -0,0 +1,61 @@
+// Copyright Epic Games, Inc. All Rights Reserved.
+
+#pragma once
+#include "INetworkPredictionProvider.h"
+
+
+class FNetworkPredictionProvider : public INetworkPredictionProvider
+{
+public:
+	static FName ProviderName;
+
+	FNetworkPredictionProvider(Trace::IAnalysisSession& InSession);
+
+	// -----------------------------------------------------
+	virtual uint32 GetNetworkPredictionTraceVersion() const override { return NetworkPredictionTraceVersion; }
+	virtual uint64 GetNetworkPredictionDataCounter() const { return DataCounter; }
+
+	// -----------------------------------------------------
+	
+	void IncrementDataCounter() { DataCounter++; }
+	void SetNetworkPredictionTraceVersion(uint32 Version);
+
+	// ------------------------------------------------------------------------------
+	//
+	//	The assumption is always that we are writing *forward* in time. This will blow up 
+	//	if we expect to be able to write frames in the past
+	//
+	// ------------------------------------------------------------------------------
+
+	FSimulationData::FConst& WriteSimulationCreated(uint32 SimulationId);
+	void WriteSimulationTick(uint32 SimulationId, FSimulationData::FTick&& InTick);
+	FSimulationData::FEngineFrame& WriteSimulationEOF(uint32 SimulationId);
+	void WriteNetRecv(uint32 SimulationId, FSimulationData::FNetSerializeRecv&& InNetRecv);
+	void WriteNetCommit(uint32 SimulationId);
+	void WriteNetFault(uint32 SimulationId);
+	void WriteOOBStateMod(uint32 SimulationId);
+	void WriteProduceInput(uint32 SimulationId);
+	void WriteSynthInput(uint32 SimulationId);
+	void WriteUserState(uint32 SimulationId, int32 Frame, uint64 EngineFrame, ENP_UserState Type, const TCHAR* UserStr);
+	void WritePIEStart();
+	
+	void WriteSimulationNetRole(uint32 SimulationId, uint64 EngineFrame, ENP_NetRole Role);
+	void WriteSimulationNetGUID(uint32 SimulationId, uint32 NetGUID);
+
+	virtual TArrayView<const TSharedRef<FSimulationData>> ReadSimulationData() const override
+	{
+		return TArrayView<const TSharedRef<FSimulationData>>(ProviderData); 
+	}
+
+private:
+
+	TSharedRef<FSimulationData>& FindOrAdd(uint32 SimulationId);
+
+	// Array that stores all of our traced simulation data
+	TArray<TSharedRef<FSimulationData>> ProviderData;
+
+	uint64 DataCounter = 0;
+	uint32 NetworkPredictionTraceVersion = 0;
+	int32 PIESessionCounter = 0;
+	Trace::IAnalysisSession& Session;
+};
\ No newline at end of file
diff --git a/Engine/Plugins/Runtime/NetworkPredictionInsights/Source/NetworkPredictionInsights/Private/NetworkPredictionTraceModule.cpp b/Engine/Plugins/Runtime/NetworkPredictionInsights/Source/NetworkPredictionInsights/Private/NetworkPredictionTraceModule.cpp
new file mode 100644
--- /dev/null
+++ b/Engine/Plugins/Runtime/NetworkPredictionInsights/Source/NetworkPredictionInsights/Private/NetworkPredictionTraceModule.cpp
@@ -0,0 +1,32 @@
+// Copyright Epic Games, Inc. All Rights Reserved.
+
+#include "NetworkPredictionTraceModule.h"
+#include "NetworkPredictionProvider.h"
+#include "NetworkPredictionAnalyzer.h"
+
+FName FNetworkPredictionTraceModule::ModuleName("NetworkPredictionTrace");
+
+void FNetworkPredictionTraceModule::GetModuleInfo(Trace::FModuleInfo& OutModuleInfo)
+{
+	OutModuleInfo.Name = ModuleName;
+	OutModuleInfo.DisplayName = TEXT("NetworkPrediction");
+}
+
+void FNetworkPredictionTraceModule::OnAnalysisBegin(Trace::IAnalysisSession& InSession)
+{
+	FNetworkPredictionProvider* NetworkPredictionProvider = new FNetworkPredictionProvider(InSession);
+	InSession.AddProvider(FNetworkPredictionProvider::ProviderName, NetworkPredictionProvider);
+	
+	InSession.AddAnalyzer(new FNetworkPredictionAnalyzer(InSession, *NetworkPredictionProvider));
+}
+
+void FNetworkPredictionTraceModule::GetLoggers(TArray<const TCHAR *>& OutLoggers)
+{
+	OutLoggers.Add(TEXT("NetworkPrediction"));
+}
+
+void FNetworkPredictionTraceModule::GenerateReports(const Trace::IAnalysisSession& Session, const TCHAR* CmdLine, const TCHAR* OutputDirectory)
+{
+
+}
+
diff --git a/Engine/Plugins/Runtime/NetworkPredictionInsights/Source/NetworkPredictionInsights/Private/NetworkPredictionTraceModule.h b/Engine/Plugins/Runtime/NetworkPredictionInsights/Source/NetworkPredictionInsights/Private/NetworkPredictionTraceModule.h
new file mode 100644
--- /dev/null
+++ b/Engine/Plugins/Runtime/NetworkPredictionInsights/Source/NetworkPredictionInsights/Private/NetworkPredictionTraceModule.h
@@ -0,0 +1,23 @@
+// Copyright Epic Games, Inc. All Rights Reserved.
+
+#pragma once
+
+#include "CoreMinimal.h"
+#include "TraceServices/ModuleService.h"
+
+class IAnimationProvider;
+namespace Insights { class ITimingViewSession; }
+
+class FNetworkPredictionTraceModule : public Trace::IModule
+{
+public:
+	// Trace::IModule interface
+	virtual void GetModuleInfo(Trace::FModuleInfo& OutModuleInfo) override;
+	virtual void OnAnalysisBegin(Trace::IAnalysisSession& Session) override;
+	virtual void GetLoggers(TArray<const TCHAR *>& OutLoggers) override;
+	virtual void GenerateReports(const Trace::IAnalysisSession& Session, const TCHAR* CmdLine, const TCHAR* OutputDirectory) override;
+	virtual const TCHAR* GetCommandLineArgument() override { return TEXT("objecttrace"); }
+
+private:
+	static FName ModuleName;
+};
\ No newline at end of file
diff --git a/Engine/Plugins/Runtime/NetworkPredictionInsights/Source/NetworkPredictionInsights/Private/UI/NPAxisViewportDouble.cpp b/Engine/Plugins/Runtime/NetworkPredictionInsights/Source/NetworkPredictionInsights/Private/UI/NPAxisViewportDouble.cpp
new file mode 100644
--- /dev/null
+++ b/Engine/Plugins/Runtime/NetworkPredictionInsights/Source/NetworkPredictionInsights/Private/UI/NPAxisViewportDouble.cpp
@@ -0,0 +1,161 @@
+// Copyright Epic Games, Inc. All Rights Reserved.
+
+// TEMPORARY copy of Insight's AxisViewportInt32 until Dev-Core copy up is complete and it becomes public
+
+#include "NPAxisViewportDouble.h"
+
+#include "Widgets/Layout/SScrollBar.h"
+
+////////////////////////////////////////////////////////////////////////////////////////////////////
+
+bool FNPAxisViewportDouble::SetSize(const float InSize)
+{
+	if (!FMath::IsNearlyEqual(Size, InSize, SLATE_UNITS_TOLERANCE))
+	{
+		Size = InSize;
+		OnSizeChanged();
+		return true;
+	}
+	return false;
+}
+
+////////////////////////////////////////////////////////////////////////////////////////////////////
+
+bool FNPAxisViewportDouble::SetScale(const double NewScale)
+{
+	const double LocalNewScale = FMath::Clamp(NewScale, MinScale, MaxScale);
+	if (LocalNewScale != Scale)
+	{
+		Scale = LocalNewScale;
+		OnScaleChanged();
+		return true;
+	}
+	return false;
+}
+
+////////////////////////////////////////////////////////////////////////////////////////////////////
+
+bool FNPAxisViewportDouble::ZoomWithFixedOffset(const double NewScale, const float Offset)
+{
+	const double LocalNewScale = FMath::Clamp(NewScale, MinScale, MaxScale);
+	if (LocalNewScale != Scale)
+	{
+		// Value at Offset should remain the same. So we resolve equation:
+		//   (Position + Offset) / Scale == (NewPosition + Offset) / NewScale
+		//   ==> NewPosition = (Position + Offset) / Scale * NewScale - Offset
+		Position = (Position + Offset) * static_cast<float>(LocalNewScale / Scale) - Offset;
+		Scale = LocalNewScale;
+
+		OnPositionChanged();
+		OnScaleChanged();
+		return true;
+	}
+	return false;
+}
+
+////////////////////////////////////////////////////////////////////////////////////////////////////
+
+bool FNPAxisViewportDouble::RelativeZoomWithFixedOffset(const float Delta, const float Offset)
+{
+	constexpr double ZoomStep = 0.25f; // as percent
+
+	double NewScale;
+
+	if (Delta > 0)
+	{
+		NewScale = Scale * FMath::Pow(1.0f + ZoomStep, Delta);
+	}
+	else
+	{
+		NewScale = Scale * FMath::Pow(1.0f / (1.0f + ZoomStep), -Delta);
+	}
+
+	return ZoomWithFixedOffset(NewScale, Offset);
+}
+
+////////////////////////////////////////////////////////////////////////////////////////////////////
+
+bool FNPAxisViewportDouble::UpdatePosWithinLimits()
+{
+	float MinPos, MaxPos;
+	GetScrollLimits(MinPos, MaxPos);
+	return EnforceScrollLimits(MinPos, MaxPos, 0.5);
+}
+
+////////////////////////////////////////////////////////////////////////////////////////////////////
+
+void FNPAxisViewportDouble::GetScrollLimits(float& OutMinPos, float& OutMaxPos)
+{
+	if (MaxPosition - MinPosition < Size)
+	{
+		OutMinPos = MaxPosition - Size;
+		OutMaxPos = MinPosition;
+	}
+	else
+	{
+		constexpr float ExtraSizeFactor = 0.15f; // allow extra 15% on sides
+		OutMinPos = MinPosition - ExtraSizeFactor * Size;
+		OutMaxPos = MaxPosition - (1.0f - ExtraSizeFactor) * Size;
+	}
+}
+
+////////////////////////////////////////////////////////////////////////////////////////////////////
+
+bool FNPAxisViewportDouble::EnforceScrollLimits(const float InMinPos, const float InMaxPos, const float InterpolationFactor)
+{
+	float Pos = Position;
+
+	if (Pos < InMinPos)
+	{
+		Pos = InterpolationFactor * Pos + (1.0f - InterpolationFactor) * InMinPos;
+
+		if (FMath::IsNearlyEqual(Pos, InMinPos, 0.5f))
+		{
+			Pos = InMinPos;
+		}
+	}
+	else if (Pos > InMaxPos)
+	{
+		Pos = InterpolationFactor * Pos + (1.0f - InterpolationFactor) * InMaxPos;
+
+		if (FMath::IsNearlyEqual(Pos, InMaxPos, 0.5f))
+		{
+			Pos = InMaxPos;
+		}
+
+		if (Pos < InMinPos)
+		{
+			Pos = InMinPos;
+		}
+	}
+
+	return ScrollAtPos(Pos);
+}
+
+////////////////////////////////////////////////////////////////////////////////////////////////////
+
+void FNPAxisViewportDouble::OnUserScrolled(TSharedPtr<SScrollBar> ScrollBar, float ScrollOffset)
+{
+	const float SX = 1.0f / (MaxPosition - MinPosition);
+	const float ThumbSizeFraction = FMath::Clamp<float>(Size * SX, 0.0f, 1.0f);
+	const float OffsetFraction = FMath::Clamp<float>(ScrollOffset, 0.0f, 1.0f - ThumbSizeFraction);
+
+	const float Pos = MinPosition + OffsetFraction * (MaxPosition - MinPosition);
+	ScrollAtPos(Pos);
+
+	ScrollBar->SetState(OffsetFraction, ThumbSizeFraction);
+}
+
+////////////////////////////////////////////////////////////////////////////////////////////////////
+
+void FNPAxisViewportDouble::UpdateScrollBar(TSharedPtr<SScrollBar> ScrollBar) const
+{
+	const float SX = 1.0f / (MaxPosition - MinPosition);
+	const float ThumbSizeFraction = FMath::Clamp<float>(Size * SX, 0.0f, 1.0f);
+	const float ScrollOffset = (Position - MinPosition) * SX;
+	const float OffsetFraction = FMath::Clamp<float>(ScrollOffset, 0.0f, 1.0f - ThumbSizeFraction);
+
+	ScrollBar->SetState(OffsetFraction, ThumbSizeFraction);
+}
+
+////////////////////////////////////////////////////////////////////////////////////////////////////
diff --git a/Engine/Plugins/Runtime/NetworkPredictionInsights/Source/NetworkPredictionInsights/Private/UI/NPAxisViewportDouble.h b/Engine/Plugins/Runtime/NetworkPredictionInsights/Source/NetworkPredictionInsights/Private/UI/NPAxisViewportDouble.h
new file mode 100644
--- /dev/null
+++ b/Engine/Plugins/Runtime/NetworkPredictionInsights/Source/NetworkPredictionInsights/Private/UI/NPAxisViewportDouble.h
@@ -0,0 +1,193 @@
+// Copyright Epic Games, Inc. All Rights Reserved.
+
+// TEMPORARY copy of Insight's AxisViewportInt32 until Dev-Core copy up is complete and it becomes public
+
+#pragma once
+
+#include "CoreMinimal.h"
+
+class SScrollBar;
+
+////////////////////////////////////////////////////////////////////////////////////////////////////
+
+class FNPAxisViewportDouble
+{
+private:
+	static constexpr float SLATE_UNITS_TOLERANCE = 0.1f;
+
+public:
+	FNPAxisViewportDouble()
+	{
+		Reset();
+	}
+
+	void Reset()
+	{
+		Size = 0.0f;
+
+		MinValue = 0.0;
+		MaxValue = 0.0;
+
+		MinPosition = 0.0f;
+		MaxPosition = 0.0f;
+		Position = 0.0f;
+
+		MinScale = 0.01;
+		MaxScale = 100.0;
+		Scale = 1.0;
+	}
+
+	FString ToDebugString(const TCHAR* Sufix) const
+	{
+		return FString::Printf(TEXT("Scale%s: %g, Pos%s: %.2f"), Sufix, Scale, Sufix, Position);
+	}
+
+	/**
+	 * Size of viewport, in pixels (Slate units).
+	 * The viewport's width if this is a horizontal axis or the viewport's height if this is a vertical axis.
+	 */
+	float GetSize() const { return Size; }
+	bool SetSize(const float InSize);
+
+	int32 GetMinValue() const { return MinValue; }
+	int32 GetMaxValue() const { return MaxValue; }
+
+	float GetMinPos() const { return MinPosition; }
+	float GetMaxPos() const { return MaxPosition; }
+	float GetPos() const { return Position; }
+
+	double GetMinScale() const { return MinScale; }
+	double GetMaxScale() const { return MaxScale; }
+
+	/** Current scale factor between Value units and Slate units (pixels), in pixels per value. [px/value_unit] */
+	double GetScale() const { return Scale; }
+
+	void SetMinMaxValueInterval(const double InMinValue, const double InMaxValue)
+	{
+		MinValue = InMinValue;
+		MaxValue = InMaxValue;
+		UpdateMinMax();
+	}
+
+	double GetValueAtPos(const float Pos) const
+	{
+		return static_cast<double>(Pos) / Scale;
+	}
+
+	double GetValueAtOffset(const float Offset) const
+	{
+		return static_cast<double>(Position + Offset) / Scale;
+	}
+
+	float GetPosForValue(const double Value) const
+	{
+		return static_cast<float>(Value * Scale);
+	}
+
+	float GetOffsetForValue(const double Value) const
+	{
+		return static_cast<float>(Value * Scale) - Position;
+	}
+
+	float GetRoundedOffsetForValue(const double Value) const
+	{
+		return FMath::RoundToFloat(static_cast<float>(Value * Scale) - Position);
+	}
+
+	bool ScrollAtPos(const float Pos)
+	{
+		if (Position != Pos)
+		{
+			Position = Pos;
+			OnPositionChanged();
+			return true;
+		}
+		return false;
+	}
+
+	bool ScrollAtValue(const double Value)
+	{
+		return ScrollAtPos(GetPosForValue(Value));
+	}
+
+	bool CenterOnValue(const double Value)
+	{
+		return ScrollAtPos(static_cast<float>(Value * Scale) - Size / 2.0f);
+	}
+
+	bool CenterOnValueInterval(const double IntervalStartValue, const double IntervalEndValue)
+	{
+		const float IntervalSize = static_cast<float>(Scale * (IntervalEndValue- IntervalStartValue));
+		if (IntervalSize > Size)
+		{
+			return ScrollAtValue(IntervalStartValue);
+		}
+		else
+		{
+			return ScrollAtPos(static_cast<float>(IntervalStartValue * Scale) - (Size - IntervalSize) / 2.0f);
+		}
+	}
+
+	void SetScaleLimits(const double InMinScale, const double InMaxScale)
+	{
+		MinScale = InMinScale;
+		MaxScale = InMaxScale;
+	}
+
+	bool SetScale(const double NewScale);
+	bool ZoomWithFixedOffset(const double NewScale, const float Offset);
+	bool RelativeZoomWithFixedOffset(const float Delta, const float Offset);
+
+	void GetScrollLimits(float& OutMinPos, float& OutMaxPos);
+	bool EnforceScrollLimits(const float InMinPos, const float InMaxPos, const float InterpolationFactor);
+	bool UpdatePosWithinLimits();
+
+	void OnUserScrolled(TSharedPtr<SScrollBar> ScrollBar, float ScrollOffset);
+	void UpdateScrollBar(TSharedPtr<SScrollBar> ScrollBar) const;
+
+private:
+	void OnSizeChanged()
+	{
+	}
+
+	void OnPositionChanged()
+	{
+	}
+
+	void OnScaleChanged()
+	{
+		UpdateMinMax();
+	}
+
+	void UpdateMinMax()
+	{
+		MinPosition = GetPosForValue(MinValue);
+		MaxPosition = GetPosForValue(MaxValue);
+
+		//const float NewPosition = FMath::Clamp(Position, MinPosition, MaxPosition);
+		//if (FMath::IsNearlyEqual(NewPosition, Position, SLATE_UNITS_TOLERANCE))
+		//{
+		//	Position = NewPosition;
+		//	OnPositionChanged();
+		//}
+	}
+
+private:
+	float Size; // size of viewport, in pixels (Slate units); [px]
+
+	double MinValue; // minimum value (inclusive); [value_unit]
+	double MaxValue; // maximum value (exclusive); [value_unit]
+
+	//double StartValue; // value of viewport's left side; [value_unit]; StartValue = GetValueAtPos(Position); computed when Position or Scale changes
+	//double EndValue; // value of viewport's right side; [value_unit]; EndValue = GetValueAtPos(Position + Size); computed when Position, Size or Scale changes
+
+	float MinPosition; // minimum position (corresponding to MinValue), in pixels (Slate units); [px]
+	float MaxPosition; // maximum position (corresponding to MaxValue), in pixels (Slate units); [px]
+	float Position; // current position of the viewport, in pixels (Slate units); [px]
+
+	double MinScale; // minimum scale factor; [px/value_unit]
+	double MaxScale; // maximum scale factor; [px/value_unit]
+	double Scale; // current scale factor between Value units and Slate units (pixels); [px/value_unit]
+};
+
+////////////////////////////////////////////////////////////////////////////////////////////////////
diff --git a/Engine/Plugins/Runtime/NetworkPredictionInsights/Source/NetworkPredictionInsights/Private/UI/NPAxisViewportInt32.cpp b/Engine/Plugins/Runtime/NetworkPredictionInsights/Source/NetworkPredictionInsights/Private/UI/NPAxisViewportInt32.cpp
new file mode 100644
--- /dev/null
+++ b/Engine/Plugins/Runtime/NetworkPredictionInsights/Source/NetworkPredictionInsights/Private/UI/NPAxisViewportInt32.cpp
@@ -0,0 +1,192 @@
+// Copyright Epic Games, Inc. All Rights Reserved.
+
+// TEMPORARY copy of Insight's AxisViewportInt32 until Dev-Core copy up is complete and it becomes public
+
+#include "NPAxisViewportInt32.h"
+
+#include "Widgets/Layout/SScrollBar.h"
+
+////////////////////////////////////////////////////////////////////////////////////////////////////
+
+bool FNPAxisViewportInt32::SetSize(const float InSize)
+{
+	if (!FMath::IsNearlyEqual(Size, InSize, SLATE_UNITS_TOLERANCE))
+	{
+		Size = InSize;
+		OnSizeChanged();
+		return true;
+	}
+	return false;
+}
+
+////////////////////////////////////////////////////////////////////////////////////////////////////
+
+bool FNPAxisViewportInt32::SetScale(const float InNewScale)
+{
+	const float LocalNewScale = FMath::Clamp(InNewScale, MinScale, MaxScale);
+	if (LocalNewScale != Scale)
+	{
+		Scale = LocalNewScale;
+		OnScaleChanged();
+		return true;
+	}
+	return false;
+}
+
+////////////////////////////////////////////////////////////////////////////////////////////////////
+
+bool FNPAxisViewportInt32::ZoomWithFixedOffset(const float NewScale, const float Offset)
+{
+	const float LocalNewScale = FMath::Clamp(NewScale, MinScale, MaxScale);
+	if (LocalNewScale != Scale)
+	{
+		// Value at Offset should remain the same. So we resolve equation:
+		//   (Position + Offset) / Scale == (NewPosition + Offset) / NewScale
+		//   ==> NewPosition = (Position + Offset) / Scale * NewScale - Offset
+		Position = (Position + Offset) * (LocalNewScale / Scale) - Offset;
+		Scale = LocalNewScale;
+
+		OnPositionChanged();
+		OnScaleChanged();
+		return true;
+	}
+	return false;
+}
+
+////////////////////////////////////////////////////////////////////////////////////////////////////
+
+bool FNPAxisViewportInt32::RelativeZoomWithFixedOffset(const float Delta, const float Offset)
+{
+	constexpr float ZoomStep = 0.25f; // as percent
+
+	float NewScale;
+
+	if (Delta > 0)
+	{
+		NewScale = Scale * FMath::Pow(1.0f + ZoomStep, Delta);
+	}
+	else
+	{
+		NewScale = Scale * FMath::Pow(1.0f / (1.0f + ZoomStep), -Delta);
+	}
+
+	// Snap to integer value of either: "number of samples per pixel" or "number of pixels per sample".
+	if (NewScale < 1.0f)
+	{
+		// N sample/px; 1 px/sample
+		if (Delta > 0)
+		{
+			NewScale = 1.0f / FMath::FloorToFloat(1.0f / NewScale);
+		}
+		else
+		{
+			NewScale = 1.0f / FMath::CeilToFloat(1.0f / NewScale);
+		}
+	}
+	else
+	{
+		// 1 sample/px; N px/sample
+		if (Delta > 0)
+		{
+			NewScale = FMath::CeilToFloat(NewScale);
+		}
+		else
+		{
+			NewScale = FMath::FloorToFloat(NewScale);
+		}
+	}
+
+	return ZoomWithFixedOffset(NewScale, Offset);
+}
+
+////////////////////////////////////////////////////////////////////////////////////////////////////
+
+bool FNPAxisViewportInt32::UpdatePosWithinLimits()
+{
+	float MinPos, MaxPos;
+	GetScrollLimits(MinPos, MaxPos);
+	return EnforceScrollLimits(MinPos, MaxPos, 0.5);
+}
+
+////////////////////////////////////////////////////////////////////////////////////////////////////
+
+void FNPAxisViewportInt32::GetScrollLimits(float& OutMinPos, float& OutMaxPos)
+{
+	if (MaxPosition - MinPosition < Size)
+	{
+		OutMinPos = MaxPosition - Size;
+		OutMaxPos = MinPosition;
+	}
+	else
+	{
+		constexpr float ExtraSizeFactor = 0.15f; // allow extra 15% on sides
+		OutMinPos = MinPosition - ExtraSizeFactor * Size;
+		OutMaxPos = MaxPosition - (1.0f - ExtraSizeFactor) * Size;
+	}
+}
+
+////////////////////////////////////////////////////////////////////////////////////////////////////
+
+bool FNPAxisViewportInt32::EnforceScrollLimits(const float InMinPos, const float InMaxPos, const float InterpolationFactor)
+{
+	float Pos = Position;
+
+	if (Pos < InMinPos)
+	{
+		Pos = InterpolationFactor * Pos + (1.0f - InterpolationFactor) * InMinPos;
+
+		if (FMath::IsNearlyEqual(Pos, InMinPos, 0.5f))
+		{
+			Pos = InMinPos;
+		}
+	}
+	else if (Pos > InMaxPos)
+	{
+		Pos = InterpolationFactor * Pos + (1.0f - InterpolationFactor) * InMaxPos;
+
+		if (FMath::IsNearlyEqual(Pos, InMaxPos, 0.5f))
+		{
+			Pos = InMaxPos;
+		}
+
+		if (Pos < InMinPos)
+		{
+			Pos = InMinPos;
+		}
+	}
+
+	if (Position != Pos)
+	{
+		return ScrollAtValue(GetValueAtPos(Pos)); // snap to sample
+	}
+
+	return false;
+}
+
+////////////////////////////////////////////////////////////////////////////////////////////////////
+
+void FNPAxisViewportInt32::OnUserScrolled(TSharedPtr<SScrollBar> ScrollBar, float ScrollOffset)
+{
+	const float SX = 1.0f / (MaxPosition - MinPosition);
+	const float ThumbSizeFraction = FMath::Clamp<float>(Size * SX, 0.0f, 1.0f);
+	const float OffsetFraction = FMath::Clamp<float>(ScrollOffset, 0.0f, 1.0f - ThumbSizeFraction);
+
+	const float Pos = MinPosition + OffsetFraction * (MaxPosition - MinPosition);
+	ScrollAtPos(Pos);
+
+	ScrollBar->SetState(OffsetFraction, ThumbSizeFraction);
+}
+
+////////////////////////////////////////////////////////////////////////////////////////////////////
+
+void FNPAxisViewportInt32::UpdateScrollBar(TSharedPtr<SScrollBar> ScrollBar) const
+{
+	const float SX = 1.0f / (MaxPosition - MinPosition);
+	const float ThumbSizeFraction = FMath::Clamp<float>(Size * SX, 0.0f, 1.0f);
+	const float ScrollOffset = (Position - MinPosition) * SX;
+	const float OffsetFraction = FMath::Clamp<float>(ScrollOffset, 0.0f, 1.0f - ThumbSizeFraction);
+
+	ScrollBar->SetState(OffsetFraction, ThumbSizeFraction);
+}
+
+////////////////////////////////////////////////////////////////////////////////////////////////////
diff --git a/Engine/Plugins/Runtime/NetworkPredictionInsights/Source/NetworkPredictionInsights/Private/UI/NPAxisViewportInt32.h b/Engine/Plugins/Runtime/NetworkPredictionInsights/Source/NetworkPredictionInsights/Private/UI/NPAxisViewportInt32.h
new file mode 100644
--- /dev/null
+++ b/Engine/Plugins/Runtime/NetworkPredictionInsights/Source/NetworkPredictionInsights/Private/UI/NPAxisViewportInt32.h
@@ -0,0 +1,211 @@
+// Copyright Epic Games, Inc. All Rights Reserved.
+
+// TEMPORARY copy of Insight's AxisViewportInt32 until Dev-Core copy up is complete and it becomes public
+
+#pragma once
+
+#include "CoreMinimal.h"
+
+class SScrollBar;
+
+////////////////////////////////////////////////////////////////////////////////////////////////////
+
+class FNPAxisViewportInt32
+{
+private:
+	static constexpr float SLATE_UNITS_TOLERANCE = 0.1f;
+
+public:
+	FNPAxisViewportInt32()
+	{
+		Reset();
+	}
+
+	void Reset()
+	{
+		Size = 0.0f;
+
+		MinValue = 0;
+		MaxValue = 0;
+
+		MinPosition = 0.0f;
+		MaxPosition = 0.0f;
+		Position = 0.0f;
+
+		MinScale = 0.0001f; // 10000 [sample/px]; 1 [px/sample]
+		MaxScale = 16.0f; // 1 [sample/px]; 16 [px/sample]
+		Scale = MaxScale;
+	}
+
+	FString ToDebugString(const TCHAR* Sufix, const TCHAR* Unit) const
+	{
+		if (Scale >= 1.0f)
+		{
+			return FString::Printf(TEXT("Scale%s: %.3f (%d pixels/%s), Pos%s: %.2f"),
+				Sufix, Scale, FMath::RoundToInt(Scale), Unit, Sufix, Position);
+		}
+		else
+		{
+			return FString::Printf(TEXT("Scale%s: %.3f (%d %ss/pixel), Pos%s: %.2f"),
+				Sufix, Scale, FMath::RoundToInt(1.0f / Scale), Unit, Sufix, Position);
+		}
+	}
+
+	/**
+	 * Size of viewport, in pixels (Slate units).
+	 * The viewport's width if this is a horizontal axis or the viewport's height if this is a vertical axis.
+	 */
+	float GetSize() const { return Size; }
+	bool SetSize(const float InSize);
+
+	int32 GetMinValue() const { return MinValue; }
+	int32 GetMaxValue() const { return MaxValue; }
+
+	float GetMinPos() const { return MinPosition; }
+	float GetMaxPos() const { return MaxPosition; }
+	float GetPos() const { return Position; }
+
+	float GetMinScale() const { return MinScale; }
+	float GetMaxScale() const { return MaxScale; }
+
+	/** Current scale factor between Value units (samples) and Slate units (pixels), in pixels per sample. [px/sample]
+	 * If Scale > 1, it represents number of pixels for one sample,
+	 * otherwise, the inverse of Scale represents number of samples in one pixel.
+	 */
+	float GetScale() const { return Scale; }
+
+	/** Returns the number of pixels (Slate units) for one viewport sample. [px] */
+	float GetSampleSize() const { return FMath::Max(1.0f, FMath::RoundToFloat(Scale)); }
+
+	/** Returns the number of samples in one pixel (Slate unit). [sample] */
+	int32 GetNumSamplesPerPixel() const { return FMath::Max(1, FMath::RoundToInt(1.0f / Scale)); }
+
+	void SetMinMaxInterval(const int32 InMinValue, const int32 InMaxValue)
+	{
+		MinValue = InMinValue;
+		MaxValue = InMaxValue;
+		UpdateMinMax();
+	}
+
+	int32 GetValueAtPos(const float Pos) const
+	{
+		return FMath::RoundToInt(Pos / Scale);
+	}
+
+	int32 GetValueAtOffset(const float Offset) const
+	{
+		return FMath::RoundToInt((Position + Offset) / Scale);
+	}
+
+	float GetPosForValue(const int32 Value) const
+	{
+		return static_cast<float>(Value) * Scale;
+	}
+
+	float GetOffsetForValue(const int32 Value) const
+	{
+		return static_cast<float>(Value) * Scale - Position;
+	}
+
+	float GetRoundedOffsetForValue(const int32 Value) const
+	{
+		return FMath::RoundToFloat(static_cast<float>(Value) * Scale - Position);
+	}
+
+	bool ScrollAtPos(const float Pos)
+	{
+		if (Position != Pos)
+		{
+			Position = Pos;
+			OnPositionChanged();
+			return true;
+		}
+		return false;
+	}
+
+	bool ScrollAtValue(const int32 Value)
+	{
+		return ScrollAtPos(GetPosForValue(Value));
+	}
+
+	bool CenterOnValue(const int32 Value)
+	{
+		return ScrollAtPos(static_cast<float>(Value) * Scale - Size / 2.0f);
+	}
+
+	bool CenterOnValueInterval(const int32 IntervalStartValue, const int32 IntervalEndValue)
+	{
+		const float IntervalSize = Scale * static_cast<float>(IntervalEndValue - IntervalStartValue);
+		if (IntervalSize > Size)
+		{
+			return ScrollAtValue(IntervalStartValue);
+		}
+		else
+		{
+			return ScrollAtPos(static_cast<float>(IntervalStartValue) * Scale - (Size - IntervalSize) / 2.0f);
+		}
+	}
+
+	void SetScaleLimits(const float InMinScale, const float InMaxScale)
+	{
+		MinScale = InMinScale;
+		MaxScale = InMaxScale;
+	}
+
+	bool SetScale(const float NewScale);
+	bool ZoomWithFixedOffset(const float NewScale, const float Offset);
+	bool RelativeZoomWithFixedOffset(const float Delta, const float Offset);
+
+	void GetScrollLimits(float& OutMinPos, float& OutMaxPos);
+	bool EnforceScrollLimits(const float InMinPos, const float InMaxPos, const float InterpolationFactor);
+	bool UpdatePosWithinLimits();
+
+	void OnUserScrolled(TSharedPtr<SScrollBar> ScrollBar, float ScrollOffset);
+	void UpdateScrollBar(TSharedPtr<SScrollBar> ScrollBar) const;
+
+private:
+	void OnSizeChanged()
+	{
+	}
+
+	void OnPositionChanged()
+	{
+	}
+
+	void OnScaleChanged()
+	{
+		UpdateMinMax();
+	}
+
+	void UpdateMinMax()
+	{
+		MinPosition = GetPosForValue(MinValue);
+		MaxPosition = GetPosForValue(MaxValue);
+
+		//const float NewPosition = FMath::Clamp(Position, MinPosition, MaxPosition);
+		//if (FMath::IsNearlyEqual(NewPosition, Position, SLATE_UNITS_TOLERANCE))
+		//{
+		//	Position = NewPosition;
+		//	OnPositionChanged();
+		//}
+	}
+
+private:
+	float Size; // size of viewport (ex.: the viewport's width if this is a horizontal axis), in pixels (Slate units); [px]
+
+	int32 MinValue; // minimum value (inclusive); [sample]
+	int32 MaxValue; // maximum value (exclusive); [sample]
+
+	//int32 StartValue; // value of viewport's left side; [value_unit]; StartValue = GetValueAtPos(Position); computed when Position or Scale changes
+	//int32 EndValue; // value of viewport's right side; [value_unit]; EndValue = GetValueAtPos(Position + Size); computed when Position, Size or Scale changes
+
+	float MinPosition; // minimum position (corresponding to MinValue), in pixels (Slate units); [px]
+	float MaxPosition; // maximum position (corresponding to MaxValue), in pixels (Slate units); [px]
+	float Position; // current position of the viewport, in pixels (Slate units); [px]
+
+	float MinScale; // minimum scale factor; [px/sample]
+	float MaxScale; // maximum scale factor; [px/sample]
+	float Scale; // current scale factor between Value units (samples) and Slate units (pixels); [px/sample]
+};
+
+////////////////////////////////////////////////////////////////////////////////////////////////////
diff --git a/Engine/Plugins/Runtime/NetworkPredictionInsights/Source/NetworkPredictionInsights/Private/UI/NetworkPredictionDrawHelpers.cpp b/Engine/Plugins/Runtime/NetworkPredictionInsights/Source/NetworkPredictionInsights/Private/UI/NetworkPredictionDrawHelpers.cpp
new file mode 100644
--- /dev/null
+++ b/Engine/Plugins/Runtime/NetworkPredictionInsights/Source/NetworkPredictionInsights/Private/UI/NetworkPredictionDrawHelpers.cpp
@@ -0,0 +1,123 @@
+// Copyright Epic Games, Inc. All Rights Reserved.
+
+#include "NetworkPredictionDrawHelpers.h"
+
+#include "Fonts/FontMeasure.h"
+#include "Fonts/SlateFontInfo.h"
+
+#include "Insights/Common/PaintUtils.h"
+#include "Insights/Common/TimeUtils.h"
+
+
+void FNetworkPredictionDrawHelpers::DrawBackground(
+	const FDrawContext& DrawContext,
+	const FSlateBrush* BackgroundAreaBrush,
+	const FLinearColor& ValidAreaColor,
+	const FLinearColor& InvalidAreaColor,
+	const FLinearColor& EdgeColor,
+	const float X0,
+	const float X1,
+	const float X2,
+	const float X3,
+	const float XP,
+	const float Y,
+	const float H,
+	float& OutValidAreaX,
+	float& OutValidAreaW)
+{
+	//  <------- W ------->
+	//  X0    X1 XP X2     X3
+	//  ++++++|**1**|++++++
+	//  ++++++|**1**|++++++
+	//  ++++++|**1**|++++++
+
+	const FLinearColor RightValidArea = ValidAreaColor * 1.1f;
+
+	if (X1 >= X3 || X2 <= X0)
+	{
+		OutValidAreaX = X0;
+		OutValidAreaW = X3 - X0;
+
+		// Draw invalid area (entire view).
+		DrawContext.DrawBox(X0, Y, X3 - X0, H, BackgroundAreaBrush, InvalidAreaColor);
+	}
+	else // X1 < X3 && X2 > X0
+	{
+		if (X1 > X0)
+		{
+			// Draw invalid area (left).
+			DrawContext.DrawBox(X0, Y, X1 - X0, H, BackgroundAreaBrush, InvalidAreaColor);
+		}
+
+		if (X2 < X3)
+		{
+			// Draw invalid area (right).
+			DrawContext.DrawBox(X2 + 1.0f, Y, X3 - X2 - 1.0f, H, BackgroundAreaBrush, InvalidAreaColor);
+
+			// Draw the right edge (end time).
+			DrawContext.DrawBox(X2, Y, 1.0f, H, BackgroundAreaBrush, EdgeColor);
+		}
+
+		float ValidAreaX = FMath::Max(X1, X0);
+		float ValidAreaW = FMath::Min(X2, X3) - ValidAreaX;
+
+		if (X1 >= X0)
+		{
+			// Draw the left edge (start time).
+			DrawContext.DrawBox(X1, Y, 1.0f, H, BackgroundAreaBrush, EdgeColor);
+
+			// Adjust valid area to not overlap the left edge.
+			ValidAreaX += 1.0f;
+			ValidAreaW -= 1.0f;
+		}
+
+		if (ValidAreaW > 0.0f)
+		{
+			// Draw valid area.
+			const float LeftValidW = XP - ValidAreaX;
+			const float RightValidW = ValidAreaW - LeftValidW;
+
+			DrawContext.DrawBox(ValidAreaX, Y, LeftValidW, H, BackgroundAreaBrush, ValidAreaColor);
+
+			DrawContext.DrawBox(XP, Y, 1.0f, H, BackgroundAreaBrush, EdgeColor);
+			DrawContext.DrawBox(XP+1, Y, RightValidW, H, BackgroundAreaBrush, RightValidArea);
+		}
+
+		OutValidAreaX = ValidAreaX;
+		OutValidAreaW = ValidAreaW;
+	}
+
+	DrawContext.LayerId++;
+}
+
+
+void FNetworkPredictionDrawHelpers::DrawBackground(const FDrawContext& DrawContext,
+								  const FSlateBrush* BackgroundAreaBrush,
+								  const float X0,
+								  const float X1,
+								  const float X2,
+								  const float X3,
+								  const	float X4,
+								  const float Y,
+								  const float H)
+{
+	const FLinearColor ValidAreaColor(0.07f, 0.07f, 0.07f, 1.0f);
+	const FLinearColor InvalidAreaColor(0.1f, 0.07f, 0.07f, 1.0f);
+	const FLinearColor EdgeColor(0.05f, 0.05f, 0.05f, 1.0f);
+
+	float ValidAreaX, ValidAreaW;
+	FNetworkPredictionDrawHelpers::DrawBackground(DrawContext, BackgroundAreaBrush, ValidAreaColor, InvalidAreaColor, EdgeColor, X0, X1, X2, X3, X4, Y, H, ValidAreaX, ValidAreaW);
+}
+
+
+void FNetworkPredictionDrawHelpers::DrawFutureBackground (	const FDrawContext& DrawContext,
+										const FSlateBrush* BackgroundAreaBrush,
+										const FLinearColor& Color,
+										const float X0,
+										const float X1,
+										const float Y,
+										const float H)
+{
+	DrawContext.DrawBox(X0, Y, X1 - X0, H, BackgroundAreaBrush, Color);
+	DrawContext.LayerId++;
+}
\ No newline at end of file
diff --git a/Engine/Plugins/Runtime/NetworkPredictionInsights/Source/NetworkPredictionInsights/Private/UI/NetworkPredictionDrawHelpers.h b/Engine/Plugins/Runtime/NetworkPredictionInsights/Source/NetworkPredictionInsights/Private/UI/NetworkPredictionDrawHelpers.h
new file mode 100644
--- /dev/null
+++ b/Engine/Plugins/Runtime/NetworkPredictionInsights/Source/NetworkPredictionInsights/Private/UI/NetworkPredictionDrawHelpers.h
@@ -0,0 +1,56 @@
+// Copyright Epic Games, Inc. All Rights Reserved.
+
+#pragma once
+
+#include "CoreMinimal.h"
+#include "Fonts/SlateFontInfo.h"
+#include "Math/Color.h"
+
+struct FDrawContext;
+struct FSlateBrush;
+
+class FNetworkPredictionDrawHelpers
+{
+public:
+	/**
+	 * Draw background.
+	 * X1-X2 specifies the valid area. X0-X1 and X2-X3 specifies the invalid areas.
+	 *
+	 *  X0    X1    X2     X3
+	 *  ++++++|*****|++++++
+	 *  ++++++|*****|++++++
+	 *  ++++++|*****|++++++
+	 */
+	static void DrawBackground(const FDrawContext& DrawContext,
+							   const FSlateBrush* BackgroundAreaBrush,
+							   const FLinearColor& ValidAreaColor,
+							   const FLinearColor& InvalidAreaColor,
+							   const FLinearColor& EdgeColor,
+							   const float X0,
+							   const float X1,
+							   const float X2,
+							   const float X3,
+							   const float X4,
+							   const float Y,
+							   const float H,
+							   float& OutValidAreaX,
+							   float& OutValidAreaW);
+
+	static void DrawBackground(	const FDrawContext& DrawContext,
+								const FSlateBrush* BackgroundAreaBrush,
+								const float X0,
+								const float X1,
+								const float X2,
+								const float X3,
+								const float X4,
+								const float Y,
+								const float H);
+
+	static void DrawFutureBackground (	const FDrawContext& DrawContext,
+										const FSlateBrush* BackgroundAreaBrush,
+										const FLinearColor& Color,
+										const float X0,
+										const float X1,
+										const float Y,
+										const float H);
+};
\ No newline at end of file
diff --git a/Engine/Plugins/Runtime/NetworkPredictionInsights/Source/NetworkPredictionInsights/Private/UI/NetworkPredictionInsightsActionManager.cpp b/Engine/Plugins/Runtime/NetworkPredictionInsights/Source/NetworkPredictionInsights/Private/UI/NetworkPredictionInsightsActionManager.cpp
new file mode 100644
--- /dev/null
+++ b/Engine/Plugins/Runtime/NetworkPredictionInsights/Source/NetworkPredictionInsights/Private/UI/NetworkPredictionInsightsActionManager.cpp
@@ -0,0 +1,46 @@
+// Copyright Epic Games, Inc. All Rights Reserved.
+
+#include "NetworkPredictionInsightsActionManager.h"
+#include "NetworkPredictionInsightsManager.h"
+#include "NetworkPredictionInsightsCommands.h"
+
+////////////////////////////////////////////////////////////////////////////////////////////////////
+// Toggle Commands
+////////////////////////////////////////////////////////////////////////////////////////////////////
+
+#define IMPLEMENT_TOGGLE_COMMAND(CmdName, IsEnabled, SetIsEnabled) \
+	\
+	void FNetworkPredictionInsightsActionManager::Map_##CmdName##_Global()\
+	{\
+		This->CommandList->MapAction(This->GetCommands().CmdName, CmdName##_Custom());\
+	}\
+	\
+	const FUIAction FNetworkPredictionInsightsActionManager::CmdName##_Custom() \
+	{\
+		FUIAction UIAction;\
+		UIAction.ExecuteAction = FExecuteAction::CreateRaw(this, &FNetworkPredictionInsightsActionManager::CmdName##_Execute);\
+		UIAction.CanExecuteAction = FCanExecuteAction::CreateRaw(this, &FNetworkPredictionInsightsActionManager::CmdName##_CanExecute);\
+		UIAction.GetActionCheckState = FGetActionCheckState::CreateRaw(this, &FNetworkPredictionInsightsActionManager::CmdName##_GetCheckState);\
+		return UIAction;\
+	}\
+	\
+	void FNetworkPredictionInsightsActionManager::CmdName##_Execute()\
+	{\
+		const bool b##IsEnabled = !This->IsEnabled();\
+		This->SetIsEnabled(b##IsEnabled);\
+	}\
+	\
+	bool FNetworkPredictionInsightsActionManager::CmdName##_CanExecute() const\
+	{\
+		return FInsightsManager::Get()->GetSession().IsValid();\
+	}\
+	\
+	ECheckBoxState FNetworkPredictionInsightsActionManager::CmdName##_GetCheckState() const\
+	{\
+		const bool b##IsEnabled = This->IsEnabled();\
+		return b##IsEnabled ? ECheckBoxState::Checked : ECheckBoxState::Unchecked;\
+	}
+
+//IMPLEMENT_TOGGLE_COMMAND(ToggleAutoScrollSimulationFrames, AutoScrollSimulationFrames, SetAutoScrollSimulationFrames)
+
+#undef IMPLEMENT_TOGGLE_COMMAND
\ No newline at end of file
diff --git a/Engine/Plugins/Runtime/NetworkPredictionInsights/Source/NetworkPredictionInsights/Private/UI/NetworkPredictionInsightsActionManager.h b/Engine/Plugins/Runtime/NetworkPredictionInsights/Source/NetworkPredictionInsights/Private/UI/NetworkPredictionInsightsActionManager.h
new file mode 100644
--- /dev/null
+++ b/Engine/Plugins/Runtime/NetworkPredictionInsights/Source/NetworkPredictionInsights/Private/UI/NetworkPredictionInsightsActionManager.h
@@ -0,0 +1,43 @@
+// Copyright Epic Games, Inc. All Rights Reserved.
+
+#pragma once
+
+#include "CoreMinimal.h"
+
+/**
+ * Class that provides helper functions for the commands to avoid cluttering profiler manager with many small functions.
+ * Can't contain any variables. Directly operates on the profiler manager instance.
+ */
+class FNetworkPredictionInsightsActionManager
+{
+	friend class FNetworkPredictionInsightsManager;
+
+private:
+	/** Private constructor. */
+	FNetworkPredictionInsightsActionManager(class FNetworkPredictionInsightsManager* Instance)
+		: This(Instance)
+	{}
+
+	//////////////////////////////////////////////////
+	// Toggle Commands
+
+#define DECLARE_TOGGLE_COMMAND(CmdName)\
+public:\
+	void Map_##CmdName##_Global(); /**< Maps UI command info CmdName with the specified UI command list. */\
+	const FUIAction CmdName##_Custom(); /**< UI action for CmdName command. */\
+protected:\
+	void CmdName##_Execute(); /**< Handles FExecuteAction for CmdName. */\
+	bool CmdName##_CanExecute() const; /**< Handles FCanExecuteAction for CmdName. */\
+	ECheckBoxState CmdName##_GetCheckState() const; /**< Handles FGetActionCheckState for CmdName. */
+
+	//DECLARE_TOGGLE_COMMAND(ToggleAutoScrollSimulationFrames)
+#undef DECLARE_TOGGLE_COMMAND
+
+
+
+	//////////////////////////////////////////////////
+
+protected:
+	/** Reference to the global instance of the profiler manager. */
+	class FNetworkPredictionInsightsManager* This;
+};
\ No newline at end of file
diff --git a/Engine/Plugins/Runtime/NetworkPredictionInsights/Source/NetworkPredictionInsights/Private/UI/NetworkPredictionInsightsCommands.cpp b/Engine/Plugins/Runtime/NetworkPredictionInsights/Source/NetworkPredictionInsights/Private/UI/NetworkPredictionInsightsCommands.cpp
new file mode 100644
--- /dev/null
+++ b/Engine/Plugins/Runtime/NetworkPredictionInsights/Source/NetworkPredictionInsights/Private/UI/NetworkPredictionInsightsCommands.cpp
@@ -0,0 +1,30 @@
+// Copyright Epic Games, Inc. All Rights Reserved.
+
+#include "NetworkPredictionInsightsCommands.h"
+#include "EditorStyleSet.h"
+
+#define LOCTEXT_NAMESPACE "NetworkPredictionCommands"
+
+FNetworkPredictionInsightsCommands::FNetworkPredictionInsightsCommands()
+	: TCommands<FNetworkPredictionInsightsCommands>(
+		TEXT("NetworkPredictionCommand"), // Context name for fast lookup
+		NSLOCTEXT("Contexts", "NetworkPredictionCommand", "Network Prediction Insights"), // Localized context name for displaying
+		NAME_None, // Parent
+		FEditorStyle::GetStyleSetName() // Icon Style Set
+	)
+{
+}
+
+// UI_COMMAND takes long for the compiler to optimize
+PRAGMA_DISABLE_OPTIMIZATION
+void FNetworkPredictionInsightsCommands::RegisterCommands()
+{
+	UI_COMMAND(ToggleAutoScrollSimulationFrames, "AutoScroll", "Toggles auto scrolling in the simulation frame view when new data is received.", EUserInterfaceActionType::ToggleButton, FInputChord(EKeys::SpaceBar));
+	UI_COMMAND(NextEngineFrame, ">", "Advances Simulation view by one Engine Frame", EUserInterfaceActionType::Button, FInputChord(EKeys::Period));
+	UI_COMMAND(PrevEngineFrame, "<", "Rewinds the Simulation view by one Engine Frame", EUserInterfaceActionType::Button, FInputChord(EKeys::Comma));
+	UI_COMMAND(FirstEngineFrame, "<<", "Jump to the first frame in the Simulation View", EUserInterfaceActionType::Button, FInputChord(EKeys::Home));
+	UI_COMMAND(LastEngineFrame, ">>", "Jump to the last frame in the Simulation View", EUserInterfaceActionType::Button, FInputChord(EKeys::End));
+}
+PRAGMA_ENABLE_OPTIMIZATION
+
+#undef LOCTEXT_NAMESPACE
\ No newline at end of file
diff --git a/Engine/Plugins/Runtime/NetworkPredictionInsights/Source/NetworkPredictionInsights/Private/UI/NetworkPredictionInsightsCommands.h b/Engine/Plugins/Runtime/NetworkPredictionInsights/Source/NetworkPredictionInsights/Private/UI/NetworkPredictionInsightsCommands.h
new file mode 100644
--- /dev/null
+++ b/Engine/Plugins/Runtime/NetworkPredictionInsights/Source/NetworkPredictionInsights/Private/UI/NetworkPredictionInsightsCommands.h
@@ -0,0 +1,35 @@
+// Copyright Epic Games, Inc. All Rights Reserved.
+
+#pragma once
+
+#include "CoreMinimal.h"
+
+////////////////////////////////////////////////////////////////////////////////////////////////////
+
+class FNetworkPredictionInsightsCommands : public TCommands<FNetworkPredictionInsightsCommands>
+{
+public:
+	/** Default constructor. */
+	FNetworkPredictionInsightsCommands();
+
+	/** Initialize commands. */
+	virtual void RegisterCommands() override;
+
+public:
+	//////////////////////////////////////////////////
+	// Global commands need to implement following method:
+	//     void Map_<CommandName>_Global();
+	// Custom commands needs to implement also the following method:
+	//     const FUIAction <CommandName>_Custom(...) const;
+	//////////////////////////////////////////////////
+	
+	TSharedPtr<FUICommandInfo> ToggleAutoScrollSimulationFrames;
+	TSharedPtr<FUICommandInfo> ToggleCompactView;
+
+
+	TSharedPtr<FUICommandInfo> NextEngineFrame;
+	TSharedPtr<FUICommandInfo> PrevEngineFrame;
+	TSharedPtr<FUICommandInfo> FirstEngineFrame;
+	TSharedPtr<FUICommandInfo> LastEngineFrame;
+};
+
diff --git a/Engine/Plugins/Runtime/NetworkPredictionInsights/Source/NetworkPredictionInsights/Private/UI/NetworkPredictionInsightsManager.cpp b/Engine/Plugins/Runtime/NetworkPredictionInsights/Source/NetworkPredictionInsights/Private/UI/NetworkPredictionInsightsManager.cpp
new file mode 100644
--- /dev/null
+++ b/Engine/Plugins/Runtime/NetworkPredictionInsights/Source/NetworkPredictionInsights/Private/UI/NetworkPredictionInsightsManager.cpp
@@ -0,0 +1,82 @@
+// Copyright Epic Games, Inc. All Rights Reserved.
+
+#include "NetworkPredictionInsightsManager.h"
+
+const FName FNetworkPredictionInsightsTabs::ToolbarID("Toolbar");
+const FName FNetworkPredictionInsightsTabs::SimFrameViewID("SimFrameView");
+const FName FNetworkPredictionInsightsTabs::SimFrameContentsID("SimFrameContents");
+
+
+TSharedPtr<FNetworkPredictionInsightsManager> FNetworkPredictionInsightsManager::Instance = nullptr;
+
+////////////////////////////////////////////////////////////////////////////////////////////////////
+
+FNetworkPredictionInsightsManager::FNetworkPredictionInsightsManager()
+	//: CommandList(MakeShared<FUICommandList>())
+	: ActionManager(this)
+{
+	//FInsightsManager::Get()->GetCommandList()
+}
+
+void FNetworkPredictionInsightsManager::PostConstructor()
+{
+	// Register tick functions.
+	//OnTick = FTickerDelegate::CreateSP(this, &FNetworkPredictionManager::Tick);
+	//OnTickHandle = FTicker::GetCoreTicker().AddTicker(OnTick, 1.0f);
+
+	FNetworkPredictionInsightsCommands::Register();
+	BindCommands();
+}
+
+void FNetworkPredictionInsightsManager::BindCommands()
+{
+	//ActionManager.Map_
+	//CommandList->MapAction(GetCommands().NextEngineFrame, FUIAction::FExecuteAction
+}
+
+FNetworkPredictionInsightsManager::~FNetworkPredictionInsightsManager()
+{
+	FNetworkPredictionInsightsCommands::Unregister();
+
+	// Unregister tick function.
+	//FTicker::GetCoreTicker().RemoveTicker(OnTickHandle);
+}
+
+TSharedPtr<FNetworkPredictionInsightsManager> FNetworkPredictionInsightsManager::Get()
+{
+	return FNetworkPredictionInsightsManager::Instance;
+}
+/*
+const TSharedRef<FUICommandList> FNetworkPredictionInsightsManager::GetCommandList() const
+{
+	return CommandList;
+}
+*/
+
+const FNetworkPredictionInsightsCommands& FNetworkPredictionInsightsManager::GetCommands()
+{
+	return FNetworkPredictionInsightsCommands::Get();
+}
+
+FNetworkPredictionInsightsActionManager& FNetworkPredictionInsightsManager::GetActionManager()
+{
+	return FNetworkPredictionInsightsManager::Instance->ActionManager;
+}
+
+bool FNetworkPredictionInsightsManager::Tick(float DeltaTime)
+{
+	//SCOPE_CYCLE_COUNTER(STAT_IOPM_Tick);
+	return true;
+}
+
+void FNetworkPredictionInsightsManager::OnSessionChanged()
+{
+	for (TWeakPtr<SNPWindow> WndWeakPtr : NetworkPredictionInsightsWindows)
+	{
+		TSharedPtr<SNPWindow> Wnd = WndWeakPtr.Pin();
+		if (Wnd.IsValid())
+		{
+			Wnd->Reset();
+		}
+	}
+}
\ No newline at end of file
diff --git a/Engine/Plugins/Runtime/NetworkPredictionInsights/Source/NetworkPredictionInsights/Private/UI/NetworkPredictionInsightsManager.h b/Engine/Plugins/Runtime/NetworkPredictionInsights/Source/NetworkPredictionInsights/Private/UI/NetworkPredictionInsightsManager.h
new file mode 100644
--- /dev/null
+++ b/Engine/Plugins/Runtime/NetworkPredictionInsights/Source/NetworkPredictionInsights/Private/UI/NetworkPredictionInsightsManager.h
@@ -0,0 +1,114 @@
+// Copyright Epic Games, Inc. All Rights Reserved.
+
+#pragma once
+
+#include "CoreMinimal.h"
+#include "NetworkPredictionInsightsActionManager.h"
+//#include "Insights/InsightsManager.h"
+
+/**
+ * This class manages the Networking Profiler state and settings.
+ */
+class FNetworkPredictionInsightsManager : public TSharedFromThis<FNetworkPredictionInsightsManager>
+{
+	friend class FNetworkPredictionInsightsActionManager;
+
+public:
+	/** Creates the Networking Profiler (Networking Insights) manager, only one instance can exist. */
+	FNetworkPredictionInsightsManager();
+
+	/** Virtual destructor. */
+	virtual ~FNetworkPredictionInsightsManager();
+
+	/** Creates an instance of the profiler manager. */
+	static TSharedPtr<FNetworkPredictionInsightsManager> Initialize()
+	{
+		if (FNetworkPredictionInsightsManager::Instance.IsValid())
+		{
+			FNetworkPredictionInsightsManager::Instance.Reset();
+		}
+
+		FNetworkPredictionInsightsManager::Instance = MakeShared<FNetworkPredictionInsightsManager>();
+		FNetworkPredictionInsightsManager::Instance->PostConstructor();
+
+		return FNetworkPredictionInsightsManager::Instance;
+	}
+
+	/** Shutdowns the Networking Profiler manager. */
+	void Shutdown()
+	{
+		FNetworkPredictionInsightsManager::Instance.Reset();
+	}
+
+protected:
+	/** Finishes initialization of the profiler manager. */
+	void PostConstructor();
+
+	/** Binds our UI commands to delegates. */
+	void BindCommands();
+
+public:
+	/**
+	 * @return the global instance of the Networking Profiler (Networking Insights) manager.
+	 * This is an internal singleton and cannot be used outside ProfilerModule.
+	 * For external use:
+	 *     IProfilerModule& ProfilerModule = FModuleManager::Get().LoadModuleChecked<IProfilerModule>("Profiler");
+	 *     ProfilerModule.GetProfilerManager();
+	 */
+	static TSharedPtr<FNetworkPredictionInsightsManager> Get();
+
+	/** @returns UI command list for the Networking Profiler manager. */
+	//const TSharedRef<FUICommandList> GetCommandList() const;
+
+	/** @return an instance of the Networking Profiler commands. */
+	static const class FNetworkPredictionInsightsCommands& GetCommands();
+
+	/** @return an instance of the Networking Profiler action manager. */
+	static FNetworkPredictionInsightsActionManager& GetActionManager();
+
+	void AddProfilerWindow(const TSharedRef<SNPWindow>& InWindow)
+	{
+		NetworkPredictionInsightsWindows.Add(InWindow);
+	}
+
+	void RemoveProfilerWindow(const TSharedRef<SNPWindow>& InWindow)
+	{
+		NetworkPredictionInsightsWindows.Remove(InWindow);
+	}
+
+	/**
+	 * Converts profiler window weak pointer to a shared pointer and returns it.
+	 * Make sure the returned pointer is valid before trying to dereference it.
+	 */
+	TSharedPtr<class SNPWindow> GetProfilerWindow(int32 Index) const
+	{
+		return NetworkPredictionInsightsWindows[Index].Pin();
+	}
+
+	////////////////////////////////////////////////////////////////////////////////////////////////////
+
+	void OnSessionChanged();
+
+protected:
+	/** Updates this manager, done through FCoreTicker. */
+	bool Tick(float DeltaTime);
+
+protected:
+	/** The delegate to be invoked when this manager ticks. */
+	FTickerDelegate OnTick;
+
+	/** Handle to the registered OnTick. */
+	FDelegateHandle OnTickHandle;
+
+	/** List of UI commands for this manager. This will be filled by this and corresponding classes. */
+	//TSharedRef<FUICommandList> CommandList;
+
+	/** An instance of the Networking Profiler action manager. */
+	FNetworkPredictionInsightsActionManager ActionManager;
+
+	/** A list of weak pointers to the Networking Profiler windows. */
+	TArray<TWeakPtr<class SNPWindow>> NetworkPredictionInsightsWindows;
+
+	/** A shared pointer to the global instance of the profiler manager. */
+	static TSharedPtr<FNetworkPredictionInsightsManager> Instance;
+};
\ No newline at end of file
diff --git a/Engine/Plugins/Runtime/NetworkPredictionInsights/Source/NetworkPredictionInsights/Private/UI/SNPSimFrameContents.cpp b/Engine/Plugins/Runtime/NetworkPredictionInsights/Source/NetworkPredictionInsights/Private/UI/SNPSimFrameContents.cpp
new file mode 100644
--- /dev/null
+++ b/Engine/Plugins/Runtime/NetworkPredictionInsights/Source/NetworkPredictionInsights/Private/UI/SNPSimFrameContents.cpp
@@ -0,0 +1,347 @@
+// Copyright Epic Games, Inc. All Rights Reserved.
+
+#include "SNPSimFrameContents.h"
+#include "Widgets/Text/SRichTextBlock.h"
+
+#define LOCTEXT_NAMESPACE "SNPSimFrameContents"
+
+SNPSimFrameContents::SNPSimFrameContents()
+{
+
+}
+SNPSimFrameContents::~SNPSimFrameContents()
+{
+
+}
+
+void SNPSimFrameContents::Reset()
+{
+
+}
+
+void SNPSimFrameContents::Construct(const FArguments& InArgs, TSharedPtr<SNPWindow> InNPWindow)
+{
+	NPWindow = InNPWindow;
+	ChildSlot
+	[
+		SNew(SOverlay)
+		.Visibility(EVisibility::SelfHitTestInvisible)
+
+		+ SOverlay::Slot()
+		.HAlign(HAlign_Left)
+		.VAlign(VAlign_Top)
+		.Padding(FMargin(0, 0, 0, 0))
+		[
+			SAssignNew(ContentsHBoxPtr, SHorizontalBox)
+		]
+	];
+
+	/*
+
+			+ SHorizontalBox::Slot()
+			.AutoWidth()
+			.VAlign(VAlign_Top)
+			.HAlign(HAlign_Left)
+			.Padding(4.0f, 4.0f, 4.0f, 4.0f)
+			[
+				SAssignNew(SimInfoTextBlock, STextBlock)
+			]
+			
+			+ SHorizontalBox::Slot()
+			.AutoWidth()
+			.VAlign(VAlign_Top)
+			.HAlign(HAlign_Left)
+			.Padding(4.0f, 4.0f, 4.0f, 4.0f)
+			[
+				SAssignNew(SimTickTextBlock, STextBlock)
+			]
+
+			+ SHorizontalBox::Slot()
+			.AutoWidth()
+			.VAlign(VAlign_Top)
+			.HAlign(HAlign_Left)
+			.Padding(4.0f, 4.0f, 4.0f, 4.0f)
+			[
+				SAssignNew(InputCmdTextBlock, STextBlock)
+			]
+
+			+ SHorizontalBox::Slot()
+			.AutoWidth()
+			.VAlign(VAlign_Top)
+			.HAlign(HAlign_Left)
+			.Padding(4.0f, 4.0f, 4.0f, 4.0f)
+			[
+				SAssignNew(SyncStateTextBlock, STextBlock)
+			]
+
+			+ SHorizontalBox::Slot()
+			.AutoWidth()
+			.VAlign(VAlign_Top)
+			.HAlign(HAlign_Left)
+			.Padding(4.0f, 4.0f, 4.0f, 4.0f)
+			[
+				SAssignNew(AuxStateTextBlock, STextBlock)
+			]
+		]
+	];
+	*/
+}
+
+void SNPSimFrameContents::NotifyContentClicked(const FSimContentsView& InContent)
+{
+	struct FUserStateWidgetInfo
+	{
+		const FSimulationData::FUserState* State = nullptr;
+		const TCHAR* Heading = nullptr;
+		FLinearColor Color = FLinearColor::White;
+	};
+
+	auto MakeUserStateWidget = [&](const FUserStateWidgetInfo& Info)
+	{
+		FString HeaderText;
+		const TCHAR* UserStr = nullptr;
+		uint64 EngineFrameHyperLink = 0;
+		if (Info.State)
+		{
+			// style=\"Hyperlink\"
+			HeaderText = FString::Printf(TEXT("%s [Sim: %d. <a id=\"engine\">Engine: %d</>. %s]"), Info.Heading, Info.State->SimFrame, Info.State->EngineFrame, LexToString(Info.State->Source));
+			UserStr = Info.State->UserStr;
+			EngineFrameHyperLink = Info.State->EngineFrame;
+		}
+		else
+		{
+			HeaderText = FString::Printf(TEXT("%s [State not found!]"), Info.Heading);
+		}
+		
+		const FTextBlockStyle& TextStyle = FCoreStyle::Get().GetWidgetStyle<FTextBlockStyle>("NormalText");
+
+		return SNew(SVerticalBox)
+			+SVerticalBox::Slot()
+			.AutoHeight()
+			.HAlign(HAlign_Left)
+			.Padding(0.0f, 4.0f, 0.0f, 4.0f)
+			[
+				/*
+				SNew(STextBlock)
+				.Text(FText::FromString(HeaderText))
+				.ColorAndOpacity(Info.Color)
+				*/
+				SNew( SRichTextBlock )
+				.Text(FText::FromString(HeaderText))
+				//.TextStyle(&TextStyle)
+				//.AutoWrapText(true)
+				//.DecoratorStyleSet( &FEditorStyle::Get() )
+				.DecoratorStyleSet( &FCoreStyle::Get() )
+				+SRichTextBlock::HyperlinkDecorator( TEXT("engine"), FSlateHyperlinkRun::FOnClick::CreateLambda([this, EngineFrameHyperLink](const FSlateHyperlinkRun::FMetadata& Metadata) { NPWindow->SetEngineFrame(EngineFrameHyperLink, true); } ) )
+			]
+
+			+SVerticalBox::Slot()
+			.AutoHeight()
+			.HAlign(HAlign_Left)
+			[
+				SNew(SEditableText)
+				.IsReadOnly(true)
+				.Text(FText::FromString(UserStr))
+				.ColorAndOpacity(Info.Color)
+			];
+	};
+
+	auto AddUserStateVSlots = [&](const TArrayView<const FUserStateWidgetInfo>& Source)
+	{
+		TSharedPtr<SVerticalBox> VBox;
+		ContentsHBoxPtr->AddSlot()
+		.AutoWidth()
+		.VAlign(VAlign_Top)
+		.HAlign(HAlign_Left)
+		[
+			SAssignNew(VBox, SVerticalBox)
+		];
+
+		auto ConditionalMake = [&](const FUserStateWidgetInfo& Info)
+		{
+			if (Info.State)
+			{
+				VBox->AddSlot()
+				.AutoHeight()
+				.HAlign(HAlign_Left)
+				.VAlign(VAlign_Top)
+				.Padding(4.0f, 4.0f, 4.0f, 4.0f)
+				[
+					MakeUserStateWidget(Info)
+				];
+			}
+		};
+
+		for (const FUserStateWidgetInfo& Info : Source)
+		{
+			ConditionalMake(Info);
+		}
+	};
+
+	ContentsHBoxPtr->ClearChildren();
+	
+	// -------------------------------------------------------------
+	// Simulation Info
+	// -------------------------------------------------------------
+	FString GroupNameStr = InContent.SimView->ConstData.GroupName.ToString();
+	FString NetRoleStr = LexToString(InContent.SimView->SparseData->NetRole);
+
+	FString SimInfoText;
+	SimInfoText += FString::Format(TEXT(
+		"{0}\n"
+		"Sim Group: {1}\n"
+		"NetGUID: {2}\n"
+		"NetRole: {3}"),
+	{
+		*InContent.SimView->ConstData.DebugName,
+		*GroupNameStr,
+		InContent.SimView->ConstData.ID.NetGUID,
+		*NetRoleStr
+	});
+
+	TSharedPtr<SVerticalBox> VertBox;
+
+	ContentsHBoxPtr->AddSlot()
+	.AutoWidth()
+	.VAlign(VAlign_Top)
+	.HAlign(HAlign_Left)
+	.Padding(4.0f, 4.0f, 4.0f, 4.0f)
+	[
+		SAssignNew(VertBox, SVerticalBox)
+		+SVerticalBox::Slot()
+		.AutoHeight()
+		.HAlign(HAlign_Left)
+		.VAlign(VAlign_Top)
+		.Padding(4.0f, 4.0f, 4.0f, 4.0f)
+		[
+			SNew(STextBlock)
+			.Text(FText::FromString(SimInfoText))
+		]
+	];
+
+	// -------------------------------------------------------------
+	//	Simulation Tick
+	// -------------------------------------------------------------
+	if (InContent.SimTick)
+	{
+		FString TickText;
+
+		const FSimulationData::FTick& SimTick = *InContent.SimTick;
+
+		TickText += FString::Format(TEXT(
+			"Simulation Frame {0}\n"
+			"Start Simulation MS: {1}\n"
+			"End Simulation MS: {2}\n"
+			"Delta Simulation MS: {3}\n"
+			"Repredict: {4}\n"
+			"Confirmed Engine Frame: {5}\n"
+			"Trashed Engine Frame: {6}\n"
+			"GFrameNumber: {7}\n\n"),
+		{
+			SimTick.OutputFrame,
+			SimTick.StartMS,
+			SimTick.EndMS,
+			(SimTick.EndMS - SimTick.StartMS),
+			SimTick.bRepredict,
+			SimTick.ConfirmedEngineFrame,
+			SimTick.TrashedEngineFrame,
+			SimTick.EngineFrame
+		});
+
+		VertBox->AddSlot()
+
+		.AutoHeight()
+		.VAlign(VAlign_Top)
+		.HAlign(HAlign_Left)
+		.Padding(4.0f, 4.0f, 4.0f, 4.0f)
+		[
+			SNew(STextBlock)
+			.Text(FText::FromString(TickText))
+		];		
+		
+		const int32 InputFrame = SimTick.OutputFrame - 1;
+
+		const FSimulationData::FUserState* InputCmd = InContent.SimView->UserData.Get(ENP_UserState::Input, InputFrame, InContent.SimTick->EngineFrame, (uint8)ENP_UserStateSource::NetRecv);
+		
+		const FSimulationData::FUserState* InSyncState = InContent.SimView->UserData.Get(ENP_UserState::Sync, InputFrame, InContent.SimTick->EngineFrame, (uint8)ENP_UserStateSource::NetRecv);
+		const FSimulationData::FUserState* InAuxState = InContent.SimView->UserData.Get(ENP_UserState::Aux, InputFrame, InContent.SimTick->EngineFrame, (uint8)ENP_UserStateSource::NetRecv);
+
+		const FSimulationData::FUserState* OutSyncState = InContent.SimView->UserData.Get(ENP_UserState::Sync, SimTick.OutputFrame, InContent.SimTick->EngineFrame, (uint8)ENP_UserStateSource::NetRecv);
+		const FSimulationData::FUserState* OutAuxState = InContent.SimView->UserData.Get(ENP_UserState::Aux, SimTick.OutputFrame, InContent.SimTick->EngineFrame, (uint8)ENP_UserStateSource::NetRecv);
+				
+
+		AddUserStateVSlots( { FUserStateWidgetInfo{ InputCmd, TEXT("Input Cmd"), FLinearColor::White } } );
+		AddUserStateVSlots( { FUserStateWidgetInfo{ InSyncState, TEXT("In Sync"), FLinearColor::White }, FUserStateWidgetInfo{ InAuxState, TEXT("In Aux"), FLinearColor::White } });
+
+		if (InAuxState != OutAuxState)
+		{
+			AddUserStateVSlots( { FUserStateWidgetInfo{ OutSyncState, TEXT("Out Sync"), FLinearColor::White }, FUserStateWidgetInfo{ OutAuxState, TEXT("Out Aux"), FLinearColor::White } } );
+		}
+		else
+		{
+			AddUserStateVSlots( { FUserStateWidgetInfo{ OutSyncState, TEXT("Out Sync"), FLinearColor::White }, FUserStateWidgetInfo{ OutAuxState, TEXT("Out Aux (No Change)"), FLinearColor::Gray} } );
+		}
+	}
+
+	// -------------------------------------------------------------
+	//	Net Receive
+	// -------------------------------------------------------------
+
+	if (InContent.NetRecv)
+	{
+		const FSimulationData::FNetSerializeRecv& NetRecv = *InContent.NetRecv;
+
+		FString NetRecvText = FString::Format(TEXT(
+			"Net Receive:\n"
+			"SimFrame: {0}\n"
+			"SimTime MS: {1}\n"
+			"Status:  {2}\n"
+			"GFrameNumber: {3}\n"),
+		{
+			NetRecv.Frame,
+			NetRecv.SimTimeMS,
+			LexToString(NetRecv.Status),
+			NetRecv.EngineFrame
+		});
+
+		ContentsHBoxPtr->AddSlot()
+		.AutoWidth()
+		.VAlign(VAlign_Top)
+		.HAlign(HAlign_Left)
+		.Padding(4.0f, 0.0f, 4.0f, 0.0f)
+		[
+			SAssignNew(VertBox, SVerticalBox)
+			+SVerticalBox::Slot()
+			.AutoHeight()
+			.HAlign(HAlign_Left)
+			.VAlign(VAlign_Top)
+			.Padding(4.0f, 4.0f, 4.0f, 4.0f)
+			[
+				SNew(STextBlock)
+				.Text(FText::FromString(NetRecvText))
+			]
+		];
+
+
+		uint8 Mask = ~((uint8)ENP_UserStateSource::NetRecv | (uint8)ENP_UserStateSource::NetRecvCommit);
+		const FSimulationData::FUserState* RecvInputCmd = InContent.SimView->UserData.Get(ENP_UserState::Input, InContent.NetRecv->Frame, InContent.NetRecv->EngineFrame, Mask);
+		const FSimulationData::FUserState* RecvSyncState = InContent.SimView->UserData.Get(ENP_UserState::Sync, InContent.NetRecv->Frame, InContent.NetRecv->EngineFrame, Mask);
+		const FSimulationData::FUserState* RecvAuxState = InContent.SimView->UserData.Get(ENP_UserState::Aux, InContent.NetRecv->Frame, InContent.NetRecv->EngineFrame, Mask);
+
+		auto SelectColor = [](const FSimulationData::FUserState* UserState)
+		{
+			return UserState && UserState->Source == ENP_UserStateSource::NetRecvCommit ? FLinearColor::White : FLinearColor::Gray;
+		};
+
+		AddUserStateVSlots( { 
+			FUserStateWidgetInfo{ RecvInputCmd, TEXT("Recv InputCmd"), SelectColor(RecvInputCmd) },
+			FUserStateWidgetInfo{ RecvSyncState, TEXT("Recv Sync"), SelectColor(RecvSyncState) }, 
+			FUserStateWidgetInfo{ RecvAuxState, TEXT("Recv Aux"), SelectColor(RecvAuxState) } } 
+		);
+	}
+
+}
+
+
+
+#undef LOCTEXT_NAMESPACE
\ No newline at end of file
diff --git a/Engine/Plugins/Runtime/NetworkPredictionInsights/Source/NetworkPredictionInsights/Private/UI/SNPSimFrameContents.h b/Engine/Plugins/Runtime/NetworkPredictionInsights/Source/NetworkPredictionInsights/Private/UI/SNPSimFrameContents.h
new file mode 100644
--- /dev/null
+++ b/Engine/Plugins/Runtime/NetworkPredictionInsights/Source/NetworkPredictionInsights/Private/UI/SNPSimFrameContents.h
@@ -0,0 +1,38 @@
+// Copyright Epic Games, Inc. All Rights Reserved.
+
+#pragma once
+
+#include "CoreMinimal.h"
+#include "Layout/Geometry.h"
+#include "Styling/SlateTypes.h"
+#include "Widgets/SCompoundWidget.h"
+
+class SNPWindow;
+
+class SNPSimFrameContents : public SCompoundWidget
+{
+public:
+
+	SNPSimFrameContents();
+	virtual ~SNPSimFrameContents();
+
+	void Reset();
+
+	SLATE_BEGIN_ARGS(SNPSimFrameContents) {}
+	SLATE_END_ARGS()
+
+	void Construct(const FArguments& InArgs, TSharedPtr<SNPWindow> InNPWindow);
+
+	void NotifyContentClicked(const FSimContentsView& InView);
+
+	TSharedPtr<SHorizontalBox> ContentsHBoxPtr;
+
+	TSharedPtr<STextBlock> SimInfoTextBlock;
+	TSharedPtr<STextBlock> SimTickTextBlock;
+
+	TSharedPtr<STextBlock> InputCmdTextBlock;
+	TSharedPtr<STextBlock> SyncStateTextBlock;
+	TSharedPtr<STextBlock> AuxStateTextBlock;
+
+	TSharedPtr<SNPWindow> NPWindow;
+};
\ No newline at end of file
diff --git a/Engine/Plugins/Runtime/NetworkPredictionInsights/Source/NetworkPredictionInsights/Private/UI/SNPSimFrameView.cpp b/Engine/Plugins/Runtime/NetworkPredictionInsights/Source/NetworkPredictionInsights/Private/UI/SNPSimFrameView.cpp
new file mode 100644
--- /dev/null
+++ b/Engine/Plugins/Runtime/NetworkPredictionInsights/Source/NetworkPredictionInsights/Private/UI/SNPSimFrameView.cpp
@@ -0,0 +1,1774 @@
+// Copyright Epic Games, Inc. All Rights Reserved.
+
+#include "SNPSimFrameView.h"
+#include "NetworkPredictionDrawHelpers.h"
+
+#define LOCTEXT_NAMESPACE "NetworkPredictionInsights"
+
+FSimFrameViewDrawHelper::FSimFrameViewDrawHelper(const FDrawContext& InDrawContext, const FSimFrameViewport& InViewport)
+	: DrawContext(InDrawContext)
+	, Viewport(InViewport)
+	, WhiteBrush(FEditorStyle::Get().GetBrush("WhiteBrush"))
+	, HoveredEventBorderBrush(FEditorStyle::Get().GetBrush("HoveredEventBorder"))
+	, SelectedEventBorderBrush(FEditorStyle::Get().GetBrush("SelectedEventBorder"))
+	, SelectionFont(FCoreStyle::GetDefaultFontStyle("Regular", 8))
+{
+
+}
+
+
+void FSimFrameViewDrawHelper::DrawBackground(FSimTime Presentable) const
+{
+	const FNPAxisViewportInt32& ViewportX = Viewport.GetHorizontalAxisViewport();
+
+	const float X0 = 0.0f;
+	const float X1 = ViewportX.GetMinPos() - ViewportX.GetPos();
+	const float X2 = ViewportX.GetMaxPos() - ViewportX.GetPos();
+	const float X3 = FMath::CeilToFloat(Viewport.GetWidth());
+
+	const float Y = 0.0f;
+	const float H = FMath::CeilToFloat(Viewport.GetHeight());
+
+	const float X4 = ViewportX.GetOffsetForValue(Presentable);
+
+	FNetworkPredictionDrawHelpers::DrawBackground(DrawContext, WhiteBrush, X0, X1, X2, X3, X4, Y, H);
+}
+
+FLinearColor FSimFrameViewDrawHelper::GetFramColorByStatus(ESimFrameStatus Status, const bool bDesaturate) const
+{
+	auto GetBaseColor = [Status]() -> FLinearColor
+	{
+		constexpr float Alpha = 1.0f;
+		switch (Status)
+		{
+		case ESimFrameStatus::Predicted:
+			return FLinearColor(0.33f, 0.33f, 1.00f, Alpha);
+		case ESimFrameStatus::Repredicted:		
+			return FLinearColor(0.29f, 0.16f, 0.70f, Alpha);
+		case ESimFrameStatus::Confirmed:
+			return FLinearColor(0.33f, 1.0f, 0.33f, Alpha);
+		case ESimFrameStatus::Trashed:
+			return FLinearColor(1.0f, 0.33f, 0.33f, Alpha);		
+		case ESimFrameStatus::Abandoned:
+			return FLinearColor(0.25f, 0.25f, 0.25f, Alpha);
+		default:
+			return FLinearColor(1.0f, 1.0f, 1.0f, Alpha);
+		}
+	};
+
+	FLinearColor Color = GetBaseColor();
+	if (bDesaturate)
+	{
+		static float d = 0.2f;
+		Color = Color.Desaturate(d);
+	}
+	return Color;
+}
+
+FLinearColor FSimFrameViewDrawHelper::GetRecvColorByStatus(ENetSerializeRecvStatus Status) const
+{
+	constexpr float Alpha = 1.0f;
+	switch (Status)
+	{
+	case ENetSerializeRecvStatus::Unknown:
+		return FLinearColor(1.0f, 1.0f, 1.0f, Alpha);
+
+	case ENetSerializeRecvStatus::Confirm:
+		return FLinearColor(0.25f, 1.f, 0.25f, Alpha);
+
+	case ENetSerializeRecvStatus::Rollback:
+		return FLinearColor(1.25f, 0.25f, 0.25f, Alpha);
+
+	case ENetSerializeRecvStatus::Jump:
+		return FLinearColor(1.f, 1.f, 0.25f, Alpha);
+
+	case ENetSerializeRecvStatus::Fault:
+		return FLinearColor(1.f, 0.25f, 1.f, Alpha);
+
+	case ENetSerializeRecvStatus::Stale:
+		return FLinearColor(0.25f, 0.25f, 0.25f, Alpha);
+
+	default:
+		return FLinearColor(1.0f, 0.0f, 0.0f, Alpha);
+	}
+}
+
+void FSimFrameViewDrawHelper::DrawCached(const FSimulationFrameView& View, float MinY) const
+{
+	const FSlateFontInfo& Font = FCoreStyle::GetDefaultFontStyle("Regular", 8);
+	const TSharedRef<FSlateFontMeasure> FontMeasureService = FSlateApplication::Get().GetRenderer()->GetFontMeasureService();
+	const float MinWidthForEmbeddedText = FontMeasureService->Measure(TEXT("XXXXX"), Font).X; // Not great, would need to know highest output frame in all tracks which we currently don't cache anywhere
+
+	const FLinearColor TextColor(0.69804f, 0.69804f, 0.69804f);
+	float PulseFactor = FMath::Lerp<float>(0.25f, 1.0f, SNPSimFrameView::PulsePCT);
+	const FLinearColor PulseColor = FLinearColor(PulseFactor, PulseFactor, PulseFactor, 1.f);
+	const FLinearColor SelectedColor(1.0f, 0.5f, 0.0f);
+	const FLinearColor SelectedPulseColor(1.0f * PulseFactor, 0.5f * PulseFactor, 0.0f);
+
+	const FLinearColor SearchHighlightColor(1.0f, 1.0f, 0.25f);
+	const FLinearColor SearchHighlightPulseColor(1.0f * PulseFactor, 1.0f * PulseFactor, 0.25f);
+
+	auto SelectBorderColor = [&](const bool bPulse, const bool bSelected, const bool bSearchHighlighted, const FLinearColor& ColorFill)
+	{
+		if (bSelected)
+		{
+			return bPulse ? SelectedPulseColor : SelectedColor;
+		}
+		if (bSearchHighlighted)
+		{
+			return bPulse ? SearchHighlightPulseColor : SearchHighlightColor;
+		}
+		
+		return bPulse ? PulseColor: FLinearColor(ColorFill.R * 0.25f, ColorFill.G * 0.25f, ColorFill.B * 0.25f, ColorFill.A);
+	};
+
+	auto DrawTick = [&](const FSimulationTrack::FSubTrack::FDrawTick& Tick, const float Y, const float H)
+	{
+		const FSimulationTrack::FSubTrack::FTickSource& Source = Tick.Source;
+
+		const FLinearColor ColorFill = GetFramColorByStatus(Source.Status, Source.bDesaturate);
+		const FLinearColor ColorBorder = SelectBorderColor(Source.bPulse, Source.bSelected, Source.bSearchHighlighted, ColorFill);
+		const float X = Tick.X;
+		const float W = Tick.W;
+
+		if ( W > 2.0f)
+		{
+			DrawContext.DrawBox(X + 1.0f, Y + 1.0f, W - 2.0f, H - 2.0f, WhiteBrush, ColorFill);
+
+			// Draw border.
+			const float B = (Source.bPulse || Source.bSelected) ? 4.f : 2.f;
+			const float B2 = B/2.f;
+						
+			DrawContext.DrawBox(X, Y, B2, H, WhiteBrush, ColorBorder);
+			DrawContext.DrawBox(X + W - B2, Y, B2, H, WhiteBrush, ColorBorder);
+			DrawContext.DrawBox(X + B2, Y, W - B, B2, WhiteBrush, ColorBorder);
+			DrawContext.DrawBox(X + B2, Y + H - B2, W - B, B2, WhiteBrush, ColorBorder);
+		}
+		else
+		{
+			DrawContext.DrawBox(X, Y, W, H, WhiteBrush, ColorBorder);
+		}
+
+		if (W >= MinWidthForEmbeddedText)
+		{
+			const FString Text = FString::Printf(TEXT("%d"), Source.Tick.OutputFrame);
+			DrawContext.DrawText(X, Y, Text, Font, FLinearColor::Black);
+		}
+	};
+
+	auto DrawNetRecv = [&](const FSimulationTrack::FSubTrack::FDrawNetRecv& NetRecv, const float Y, const float H)
+	{
+		const FSimulationTrack::FSubTrack::FNetRecvSource& Source = NetRecv.Source;
+
+		const FLinearColor NetRecvColor = GetRecvColorByStatus(Source.NetRecv.Status);
+		const FLinearColor NetRecvBorder(NetRecvColor.R * 0.75f, NetRecvColor.G * 0.75f, NetRecvColor.B * 0.75f, 1.f);
+
+		const float X = NetRecv.X-1.f;
+		const float W = NetRecv.W();
+		const float NRY = Y-4.f;
+		const float NRH = H+8.f;
+
+		const bool DrawBorder = Source.bPulse || Source.bSelected || Source.bSearchHighlighted;
+		if (DrawBorder)
+		{
+			FLinearColor BorderColor = SelectBorderColor(Source.bPulse, Source.bSelected, Source.bSearchHighlighted, NetRecvColor);
+			DrawContext.DrawBox(X-1, NRY-1, W+2, NRH+2, WhiteBrush, BorderColor);
+		}
+
+		DrawContext.DrawBox(X, NRY, W, NRH, WhiteBrush, NetRecvColor);
+	};
+	
+	for (const FSimulationActorGroup& ActorGroup : View.ActorGroups)
+	{
+		if (ActorGroup.Y > MinY)
+		{
+			DrawContext.DrawText(0, ActorGroup.Y, ActorGroup.DisplayString, Font, TextColor);
+			DrawContext.LayerId++;
+		}
+
+		for (const FSimulationTrack& Track : ActorGroup.SimulationTracks)
+		{
+			if (Track.Y > MinY)
+			{
+				DrawContext.DrawText(0, Track.Y, Track.DisplayString, Font, TextColor);
+				DrawContext.LayerId++;
+			}
+
+			// Orphan recvs first so they draw behind everything	
+			if (Track.SubTracks.Num() > 0 && Track.SubTracks[0].Y > MinY)
+			{
+				const float Y = Track.SubTracks[0].Y;
+				const float H = Track.SubTracks[0].H;
+				for (const FSimulationTrack::FSubTrack::FDrawNetRecv& NetRecv : Track.OrphanedDrawNetRecvList)
+				{
+					DrawNetRecv(NetRecv, Y, H);
+				}
+			}
+
+			for (const FSimulationTrack::FSubTrack& SubTrack : Track.SubTracks)
+			{
+				const float Y = SubTrack.Y;
+				const float H = SubTrack.H;
+
+				for (const FSimulationTrack::FSubTrack::FDrawTick& Tick : SubTrack.DrawTickList)
+				{
+					DrawTick(Tick, Y, H);
+				}
+				
+				for (const FSimulationTrack::FSubTrack::FDrawNetRecv& NetRecv : SubTrack.DrawNetRecvList)
+				{
+					DrawNetRecv(NetRecv, Y, H);
+				}
+			}
+		}
+	}
+}
+
+// --------------------------------------------------------------------
+
+SNPSimFrameView::SNPSimFrameView()
+{
+
+}
+
+SNPSimFrameView::~SNPSimFrameView()
+{
+
+}
+
+void SNPSimFrameView::Construct(const FArguments& InArgs, TSharedPtr<SNPWindow> InNetworkPredictionWindow)
+{
+	NetworkPredictionWindow = InNetworkPredictionWindow;
+	InsightsModule = &FModuleManager::LoadModuleChecked<IUnrealInsightsModule>("TraceInsights");
+
+	ChildSlot
+	[
+		SNew(SOverlay)
+		.Visibility(EVisibility::SelfHitTestInvisible)
+
+		+ SOverlay::Slot()
+		.VAlign(VAlign_Top)
+		.Padding(FMargin(0, 0, 0, 0))
+		[
+			SAssignNew(HorizontalScrollBar, SScrollBar)
+			.Orientation(Orient_Horizontal)
+			.AlwaysShowScrollbar(false)
+			.Visibility(EVisibility::Visible)
+			.Thickness(FVector2D(5.0f, 5.0f))
+			.RenderOpacity(0.75)
+			.OnUserScrolled(this, &SNPSimFrameView::HorizontalScrollBar_OnUserScrolled)
+		]
+
+		+ SOverlay::Slot()
+		.HAlign(HAlign_Right)
+		.Padding(FMargin(0, 0, 0, 0))
+		[
+			SAssignNew(VerticalScrollBar, SScrollBar)
+			.Orientation(Orient_Vertical)
+			.AlwaysShowScrollbar(false)
+			.Visibility(EVisibility::Visible)
+			.Thickness(FVector2D(5.0f, 5.0f))
+			.RenderOpacity(0.75)
+			.OnUserScrolled(this, &SNPSimFrameView::VerticalScrollBar_OnUserScrolled)
+		]
+	];
+
+	UpdateHorizontalScrollBar();
+	UpdateVerticalScrollBar();
+
+	Reset();
+
+	NetworkPredictionWindow->OnFilteredDataCollectionChange.AddLambda([this](const FFilteredDataCollection& FilteredData)
+	{
+		bIsStateDirty = true; 
+		if (bAutoScroll)
+		{
+			bAutoScrollDirty = true;
+		}
+	});
+}
+
+float SNPSimFrameView::PulsePCT;
+
+void SNPSimFrameView::Tick(const FGeometry& AllottedGeometry, const double InCurrentTime, const float InDeltaTime)
+{
+	PulseAccumulator += InDeltaTime;
+	const float SinIn = PulseAccumulator * (2.f * PI);
+	PulsePCT = (FMath::Sin(SinIn) + 1.f) / 2.f;
+	if (PulseAccumulator > 100000.f)
+	{
+		PulseAccumulator = 0.f;
+	}
+
+	if (ThisGeometry != AllottedGeometry || bIsViewportDirty)
+	{
+		bIsViewportDirty = false;
+		const float ViewWidth = AllottedGeometry.GetLocalSize().X;
+		const float ViewHeight = AllottedGeometry.GetLocalSize().Y;
+		Viewport.SetSize(ViewWidth, ViewHeight);
+		bIsDrawDirty = true;
+	}
+
+	ThisGeometry = AllottedGeometry;
+
+	FNPAxisViewportInt32& ViewportX = Viewport.GetHorizontalAxisViewport();
+
+	if (!bIsScrolling)
+	{
+		// Elastic snap to horizontal limits.
+		if(ViewportX.UpdatePosWithinLimits())
+		{
+			bIsDrawDirty = true;
+		}
+
+		/*
+		if (Viewport.GetVerticalAxisViewport().UpdatePosWithinLimits())
+		{
+			bIsDrawDirty = true;
+		}*/
+	}
+
+	FNPAxisViewportDouble& ViewportY = Viewport.GetVerticalAxisViewport();
+
+	float MinPos = 0.f;
+	float MaxPos = FMath::Max<float>(0.f, Viewport.GetVerticalAxisViewport().GetMaxPos() - ViewportY.GetSize());
+
+	if (Viewport.GetVerticalAxisViewport().EnforceScrollLimits(MinPos, MaxPos, 0.f))
+	{
+		bIsDrawDirty = true;
+	}
+
+	/*
+	uint64 Time = FPlatformTime::Cycles64();
+	if (Time > AnalysisSyncNextTimestamp)
+	{
+		const uint64 WaitTime = static_cast<uint64>(0.1 / FPlatformTime::GetSecondsPerCycle64()); // 100ms
+		AnalysisSyncNextTimestamp = Time + WaitTime;
+
+		TSharedPtr<const Trace::IAnalysisSession> Session = FInsightsManager::Get()->GetSession();
+		if (Session.IsValid())
+		{
+			Trace::FAnalysisSessionReadScope SessionReadScope(*Session.Get());
+			const Trace::INetProfilerProvider& NetProfilerProvider = Trace::ReadNetProfilerProvider(*Session.Get());
+
+			const uint32 NewConnectionChangeCount = NetProfilerProvider.GetConnectionChangeCount();
+			if (NewConnectionChangeCount != ConnectionChangeCount)
+			{
+				ConnectionChangeCount = NewConnectionChangeCount;
+				bIsStateDirty = true;
+			}
+		}
+	}
+	*/
+
+	if (bIsStateDirty)
+	{
+		bIsStateDirty = false;
+		UpdateState();
+		bIsDrawDirty = true;
+	}
+
+	if (bIsDrawDirty || bAutoScrollDirty)
+	{
+		UpdateDraw();
+		bIsDrawDirty = false;
+	}
+}
+
+void SNPSimFrameView::UpdateState()
+{
+	check(InsightsModule);
+
+	//FStopwatch Stopwatch;
+	//Stopwatch.Start();
+	
+	FNPAxisViewportInt32& ViewportX = Viewport.GetHorizontalAxisViewport();
+	ViewportX.SetMinMaxInterval(0, 0);
+
+	const FNPAxisViewportDouble& ViewportY = Viewport.GetVerticalAxisViewport();
+	
+	// ------------------------------------------------------------------------------------------------------------------
+	//	Filter and order what we want to draw here. This will only get called when our provided data changes
+	// ------------------------------------------------------------------------------------------------------------------
+	
+	SimulationFrameView.Reset();
+
+	TArray<FSimulationActorGroup>& ActorGroups = SimulationFrameView.ActorGroups;
+
+	TSharedPtr<const Trace::IAnalysisSession> Session = InsightsModule->GetAnalysisSession();
+	if (Session.IsValid() && NetworkPredictionWindow.IsValid())
+	{
+		Trace::FAnalysisSessionReadScope SessionReadScope(*Session.Get());
+		const INetworkPredictionProvider* NetworkPredictionProvider = ReadNetworkPredictionProvider(*Session.Get());
+		if (NetworkPredictionProvider)
+		{
+			BuildSimulationView_ActorGroups(NetworkPredictionWindow->GetFilteredDataCollection());
+		}
+	}
+}
+
+// -----------------------------------------------------------------------------------------------------------
+//	Resets and fills out SimulationFrameView: Creates and sorts ActorGroups that we want to draw
+// -----------------------------------------------------------------------------------------------------------
+void SNPSimFrameView::BuildSimulationView_ActorGroups(const FFilteredDataCollection& FilteredDataCollection)
+{
+	SimulationFrameView.Reset();
+	
+	TArray<FSimulationActorGroup>& ActorGroups = SimulationFrameView.ActorGroups;
+	SimulationFrameView.HeadEngineFrame = FilteredDataCollection.LastEngineFrame;
+
+	auto InitTrack = [&](FSimulationActorGroup& Group, const TSharedRef<FSimulationData::FRestrictedView>& View)
+	{
+		Group.SimulationTracks.Emplace( FSimulationTrack{View});
+		Group.MaxAllowedSimTime = FMath::Max(Group.MaxAllowedSimTime, View->EOFState.GetLast().TotalAllowedSimTimeMS);
+		Group.MaxEngineFrame = FMath::Max(Group.MaxEngineFrame, View->EOFState.GetLast().EngineFrame);
+
+		SimulationFrameView.HeadEngineFrame = FMath::Max(SimulationFrameView.HeadEngineFrame, Group.MaxEngineFrame);
+	};
+
+	auto PopulateTracks = [&](FSimulationActorGroup& Group, ENP_NetRole SubTrackRole)
+	{
+		for (int32 idx=0; idx < FilteredDataCollection.Simulations.Num(); ++idx)
+		{
+			const auto& SimView = FilteredDataCollection.Simulations[idx];
+			const auto& SparseData = SimView->SparseData;
+
+			// Find unused match on netguid and role
+			if (SimView->ConstData.ID == Group.ID && SparseData->NetRole == SubTrackRole)
+			{
+				InitTrack(Group, SimView);
+			}
+		}
+	};
+
+	// ---------------------------------------------------------------------------------------------
+	// Build Actor Groups. Put the authority sim first, then AP, then SP
+	// ---------------------------------------------------------------------------------------------
+	
+	for (const auto& SimView : FilteredDataCollection.Simulations)
+	{
+		const auto& SparseData = SimView->SparseData;
+		if (ActorGroups.FindByPredicate([&](const FSimulationActorGroup& Existing) { return Existing.ID == SimView->ConstData.ID; }) == nullptr)
+		{
+			const bool bHasAutoProxy = FilteredDataCollection.Simulations.FindByPredicate([&](TSharedRef<FSimulationData::FRestrictedView>& SearchSimView) -> bool
+			{
+				return (SearchSimView->ConstData.ID == SimView->ConstData.ID) && (SearchSimView->SparseData->NetRole == ENP_NetRole::AutonomousProxy);
+			}) != nullptr;
+
+			ActorGroups.Emplace( FSimulationActorGroup{ SimView->ConstData.ID, SimView->ConstData.DebugName, SimView->ConstData.GroupName, bHasAutoProxy } );
+		}
+	}
+
+	ActorGroups.Sort([this](const FSimulationActorGroup& A, const FSimulationActorGroup& B)
+	{
+		// Explicit sorting first
+		int32 UserIndexA = UserSortedNetActors.IndexOfByKey(A.ID);
+		int32 UserIndexB = UserSortedNetActors.IndexOfByKey(B.ID);
+		if (UserIndexA != UserIndexB)
+		{
+			return UserIndexA > UserIndexB;
+		}
+
+		// Score them
+		int32 ScoreA = A.ID.PIESession * 10000;
+		int32 ScoreB = B.ID.PIESession * 10000;		
+
+		ScoreA += A.bHasAutoProxy ? 100 : 0;
+		ScoreB += B.bHasAutoProxy ? 100 : 0;
+		
+		ScoreA += A.ID.NetGUID < B.ID.NetGUID ? 1 : -1;
+		return ScoreA > ScoreB;
+	});
+
+	TArray<ENP_NetRole> TrackRoleOrder = {ENP_NetRole::Authority, ENP_NetRole::AutonomousProxy, ENP_NetRole::SimulatedProxy};
+
+	for (FSimulationActorGroup& Group : ActorGroups)
+	{
+		for (ENP_NetRole Role : TrackRoleOrder)
+		{
+			PopulateTracks(Group, Role);
+		}
+	}
+							
+	// ---------------------------------------------------------------------------------------------
+	// Determine what the primary GameInstance ID that everything should be anchored around is
+	// ---------------------------------------------------------------------------------------------
+	
+	TMap<int32, int32> PIESessionToAuthorityGameInstanceIDMap;
+	for (FSimulationActorGroup& Group : ActorGroups)
+	{
+		for (const FSimulationTrack& Track : Group.SimulationTracks)
+		{
+			const FSimulationData::FSparse& SparseData = Track.View->SparseData.Get();
+
+			static ENP_NetRole TargetRole = ENP_NetRole::Authority;
+			if (SparseData.NetRole == TargetRole && PIESessionToAuthorityGameInstanceIDMap.Contains(Track.View->ConstData.ID.PIESession) == false)
+			{
+				PIESessionToAuthorityGameInstanceIDMap.Add(Track.View->ConstData.ID.PIESession, Track.View->ConstData.GameInstanceId);
+				break;
+			}
+		}
+	}
+
+	// ---------------------------------------------------------------------------------------------
+	// Now that we know who to base around, calc the offset for each group
+	// This is tricky. There is no global simulation time to easily look at (since sims can tick at different, arbitrary rates)
+	// There is also no guarantee about simulation lifespans overlapping. We can have distinct
+	// "islands" of simulations.	
+	//
+	// The goal here is to align each actor group based on where the presentable game instances's simulation
+	// was at the latest engine frame that we are viewing. So SimA and SimB's last tick will be shown in alignment, 
+	// for the presentable game instance, even if SimB has been starved and not ticked for a few frames: the idea is
+	// the last tick is what we are showing and we are currently looking at a snapshot of the simulation at some Max engine frame.
+	// ---------------------------------------------------------------------------------------------
+	SimulationFrameView.PresentableTimeMS = 0;
+
+	// Represents a range or island of simulation time. We will use Engine Frames to determine these ranges
+	struct FEngineFrameRange
+	{
+		struct FViewItem
+		{
+			FSimulationActorGroup& Group;
+			TSharedRef<FSimulationData::FRestrictedView> View;
+			int32 SimTimeOffset = 0;
+			uint64 ViewMin;
+			uint64 ViewMax;
+		};
+		
+		TArray<FViewItem> Views;
+		uint64 Min; // Min of all views in this range
+		uint64 Max; // Max of all views in this range
+	};
+
+	TArray<FEngineFrameRange> PresentableRanges;
+
+	// -------------------------------------------------------
+	// Build the PresentableRanges list
+	// -------------------------------------------------------
+	for (FSimulationActorGroup& Group : ActorGroups)
+	{
+		for (const FSimulationTrack& Track : Group.SimulationTracks)
+		{
+			if (Track.View->ConstData.GameInstanceId == PIESessionToAuthorityGameInstanceIDMap.FindChecked(Group.ID.PIESession)) // We only care about the PresentableGameInstance simulation of this group
+			{
+				// Minus two is a fudge factor to avoid weirdness when processing live feeds where some sims may be ahead of others
+				const uint64 ThisTrackMin = Track.View->EOFState.GetFirst().EngineFrame;
+				const uint64 ThisTrackMax = Track.View->EOFState.GetLast().EngineFrame;
+				FEngineFrameRange* PresentableRange = PresentableRanges.FindByPredicate([ThisTrackMin, ThisTrackMax](const FEngineFrameRange& R)
+				{
+					return !(R.Max < ThisTrackMin || R.Min > ThisTrackMax);
+				});
+
+				if (PresentableRange)
+				{
+					PresentableRange->Views.Add({Group, Track.View, 0, ThisTrackMin, ThisTrackMax});
+					PresentableRange->Min = FMath::Min(PresentableRange->Min, ThisTrackMin);
+					PresentableRange->Max = FMath::Max(PresentableRange->Max, ThisTrackMax);
+				}
+				else 
+				{
+					PresentableRange = &PresentableRanges.Emplace_GetRef(FEngineFrameRange{{{Group, Track.View, 0, ThisTrackMin, ThisTrackMax}}, ThisTrackMin, ThisTrackMax});
+				}
+				
+				break;
+			}
+		}
+	}
+
+	PresentableRanges.Sort([](const FEngineFrameRange& A, const FEngineFrameRange& B)
+	{
+		return A.Max < B.Max;
+	});
+
+	// -------------------------------------------------------
+	//	Now within each range/island of time, caluclate each view within range's OffsetSimTimeMS
+	//	This is the mount of time to shift its simulation timeline so it is line with with the primary sim.
+	// -------------------------------------------------------
+	for (int32 RangeIdx=0; RangeIdx < PresentableRanges.Num(); ++RangeIdx)
+	{
+		FEngineFrameRange& Range = PresentableRanges[RangeIdx];
+		
+		// Go through each subsequent view in this range and find a common engine frame with a previous simulation track
+		for (int32 idx=1; idx < Range.Views.Num(); ++idx)
+		{
+			bool bFound = false;
+			for (int32 j=0; j < idx; ++j)
+			{
+				FEngineFrameRange::FViewItem& Current = Range.Views[idx]; // The view we are processing
+				FEngineFrameRange::FViewItem& Base = Range.Views[j];	// The attempted base view (will be prior in the list)
+
+				if (Base.ViewMax < Current.ViewMin || Base.ViewMin > Current.ViewMax)
+				{
+					// No overlap here, so move on to another
+					continue;
+				}
+				
+				auto FindIdx = [](const TRestrictedPageArrayView<FSimulationData::FEngineFrame>& EOFStateView, uint64 EngineFrame)
+				{
+					const FSimulationData::FEngineFrame* Found = nullptr;
+					for (auto It = EOFStateView.GetIteratorFromEnd(); It; --It)
+					{
+						if (It->EngineFrame == EngineFrame)
+						{
+							Found = &*It;
+							break;
+						}
+					}
+
+					return Found;
+				};
+
+				// Find common EngineFrame data
+				const uint64 Max = FMath::Min(Current.ViewMax, Base.ViewMax);
+				const FSimulationData::FEngineFrame* BaseState = FindIdx(Base.View->EOFState, Max);
+				const FSimulationData::FEngineFrame* CurrentState = FindIdx(Current.View->EOFState, Max);
+
+				check(BaseState && CurrentState);
+
+				// Compute difference in ProcessedSimulationTime at the end of this engine frame
+				const FSimTime BaseSimMS = BaseState->TotalProcessedSimTimeMS;
+				const FSimTime CurrentSimMS = CurrentState->TotalProcessedSimTimeMS;
+				const FSimTime Delta = BaseSimMS - CurrentSimMS;
+
+				// This is the final amount to offset this groups local sim time to align it properly with the primary simulation in this range
+				Current.Group.OffsetSimTimeMS = Delta;
+
+				bFound = true;
+				break;
+			}
+			
+			// Should find something, else the ranges have not been calculated right
+			check(bFound);
+		}
+	}
+
+	// Still need to offset the presentable groups themselves so that they don't overlap
+	FSimTime PrevSimTime = 0;
+	for (int32 RangeIdx=0; RangeIdx < PresentableRanges.Num(); ++RangeIdx)
+	{
+		FEngineFrameRange& Range = PresentableRanges[RangeIdx];
+		FSimTime MinOffset = TNumericLimits<int32>::Max();
+		for (FEngineFrameRange::FViewItem& View : Range.Views)
+		{
+			View.Group.OffsetSimTimeMS += PrevSimTime;
+			MinOffset = FMath::Min(MinOffset, View.Group.OffsetSimTimeMS);
+		}
+
+		PrevSimTime = Range.Views[0].View->EOFState.GetLast().TotalProcessedSimTimeMS + 1000;
+		if (MinOffset < 0)
+		{
+			PrevSimTime += -MinOffset;
+		}
+	}
+
+
+	
+	// Set presentable line, this is where the gray/light gray line is drawn
+	if (ActorGroups.Num() > 0)
+	{
+		SimulationFrameView.PresentableTimeMS = ActorGroups[0].OffsetSimTimeMS + ActorGroups[0].SimulationTracks[0].View->EOFState.GetLast().TotalProcessedSimTimeMS;
+	}
+
+	// Calc ViewportMaxSimTimeMS
+	ViewportMaxSimTimeMS = TNumericLimits<uint32>::Min();
+	ViewportMinSimTimeMS = TNumericLimits<uint32>::Max();
+	for (FSimulationActorGroup& Group : ActorGroups)
+	{
+		//Group.OffsetSimTimeMS = SimulationFrameView.PresentableTimeMS - Group.PresentableSimTimeMS;
+		ViewportMaxSimTimeMS = FMath::Max(ViewportMaxSimTimeMS, Group.MaxAllowedSimTime + Group.OffsetSimTimeMS);
+		ViewportMinSimTimeMS = FMath::Min(ViewportMinSimTimeMS, Group.OffsetSimTimeMS);
+	}
+
+	// Now build out the tracks for each actor group we created
+	BuildSimulationView_Tracks();
+}
+
+// Much simpler version of the offset calculating code but does not work for simulations with non overlapping lifespans
+#if 0
+// Now that we know who to base around, calc the presentable time for each actor group
+SimulationFrameView.PresentableTimeMS = 0;
+uint64 PresentableGFrame = 0; // Engine frame we are using to calculate presentable time
+for (FSimulationActorGroup& Group : ActorGroups)
+{
+	for (const FSimulationTrack& Track : Group.SimulationTracks)
+	{
+		if (Track.View->ConstData.GameInstanceId == PresentableGameInstanceId)
+		{
+			FSimTime LastProcessedSimTimeMS = Track.View->EOFState.GetLast().TotalProcessedSimTimeMS;
+			Group.PresentableSimTimeMS = Track.View->EOFState.GetLast().TotalProcessedSimTimeMS;
+			SimulationFrameView.PresentableTimeMS = FMath::Max(SimulationFrameView.PresentableTimeMS, Group.PresentableSimTimeMS);
+			break;
+		}
+	}
+}
+
+// Calc ViewportMaxSimTimeMS
+for (FSimulationActorGroup& Group : ActorGroups)
+{
+	Group.OffsetSimTimeMS = SimulationFrameView.PresentableTimeMS - Group.PresentableSimTimeMS;
+	ViewportMaxSimTimeMS = FMath::Max(ViewportMaxSimTimeMS, Group.MaxAllowedSimTime + Group.OffsetSimTimeMS);
+}
+
+#endif
+
+// -----------------------------------------------------------------------------------------------------------
+//	Builds out the actual Simulation Tracks and SubTracks that each group will display.
+//	Does everything up until we need screen space information
+// -----------------------------------------------------------------------------------------------------------
+void SNPSimFrameView::BuildSimulationView_Tracks()
+{
+	TArray<FSimulationActorGroup>& ActorGroups = SimulationFrameView.ActorGroups;
+	const FSimulationData::FTick* SelectedTick = NetworkPredictionWindow->GetSelectedContent().SimTick;
+	const FSimulationData::FNetSerializeRecv* SelectedNetRecv = NetworkPredictionWindow->GetSelectedContent().NetRecv;
+
+	for (FSimulationActorGroup& Group : ActorGroups)
+	{
+		auto ProcessNetRecv = [&](const FSimulationData::FNetSerializeRecv& NetRecv, TArray<FSimulationTrack::FSubTrack::FNetRecvSource>& NetRecvSourceList, const FSimulationData::FRestrictedView& SimView)
+		{
+			if (NetRecv.EngineFrame > SimulationFrameView.HeadEngineFrame)
+			{
+				return;
+			}
+
+			const int32 NetMS = NetRecv.SimTimeMS + Group.OffsetSimTimeMS;
+
+			const bool bPulseNetRecv = NetRecv.EngineFrame == SimulationFrameView.HeadEngineFrame;
+			const bool bSelected = SelectedNetRecv == &NetRecv;
+			const bool bSearchHighlighted = PerformSearch(NetRecv, SimView);
+
+			NetRecvSourceList.Emplace(FSimulationTrack::FSubTrack::FNetRecvSource{NetRecv, NetMS, bPulseNetRecv, bSelected, bSearchHighlighted});
+		};
+
+		Group.DisplayString = FString::Printf(TEXT("%s [NetGUID: %d] %s"), *Group.DebugName, Group.ID.NetGUID, *LexToString(Group.MaxAllowedSimTime));
+
+		for (FSimulationTrack& Track : Group.SimulationTracks)
+		{
+			const FSimulationData::FRestrictedView& SimView = Track.View.Get();
+			const FSimulationData::FSparse& SparseData = SimView.SparseData.Get();
+
+			Track.DisplayString = FString::Printf(TEXT("   %s"), LexToString(SparseData.NetRole));
+			Track.SubTracks.Reset();
+
+			FSimulationTrack::FSubTrack* SubTrack = &Track.SubTracks.Emplace_GetRef();
+			uint64 PrevEngineFrame = 0;
+			bool bPrevDesat = false;
+
+			for (auto It = SimView.Ticks.GetIterator(); It; ++It)
+			{
+				const FSimulationData::FTick& Tick = *It;
+
+				const FSimTime StartMS = Tick.StartMS + Group.OffsetSimTimeMS;
+				const FSimTime EndMS = Tick.EndMS + Group.OffsetSimTimeMS;
+
+				// -------------------------------------------------------------------
+				//	Get the SubTrack that this tick should go on
+				// -------------------------------------------------------------------
+
+				const bool bFitsPrevSubTrack = SubTrack ? (SubTrack->PrevMS <= StartMS) : false;
+
+				// Search for a new SubTrack if we don't fit on the last one or if we not in linear mode (linear mode = stay on the same subtrack unless forced off)
+				if (!bFitsPrevSubTrack || !bLinearSimFrameView)
+				{
+					if (bCompactSimFrameView)
+					{
+						// Find first one that can fit us, starting from the top. Creates a zig zag style pattern
+						SubTrack = Track.SubTracks.FindByPredicate([&StartMS](const FSimulationTrack::FSubTrack& ST)
+						{
+							return ST.PrevMS <= StartMS;
+						});
+					}
+					else
+					{
+						// Check [0] first but then only consider tracks below us. Keeps a "cascade" style pattern
+						auto CurrentSubTrack = SubTrack;
+						SubTrack = nullptr;
+
+						bool bPassedCurrent = false;
+						for (int32 idx=0; idx < Track.SubTracks.Num(); ++idx)
+						{
+							if (&Track.SubTracks[idx] == CurrentSubTrack)
+							{
+								bPassedCurrent = true;
+							}
+
+							if (idx > 0 && !bPassedCurrent)
+							{
+								continue;
+							}
+
+							if (Track.SubTracks[idx].PrevMS <= StartMS)
+							{
+								SubTrack = &Track.SubTracks[idx];
+								break;
+							}
+						}
+					}
+				}
+
+				// No suitable track found, create a new one
+				if (SubTrack == nullptr)
+				{
+					SubTrack = &Track.SubTracks.Emplace_GetRef();
+				}
+
+				SubTrack->PrevMS = EndMS;
+
+				// -------------------------------------------------------------------
+				//	Finalize the tick
+				// -------------------------------------------------------------------
+
+				// Toggle desat flag if we are on a new engine frame
+				if (PrevEngineFrame != Tick.EngineFrame)
+				{
+					bPrevDesat = !bPrevDesat;
+					PrevEngineFrame = Tick.EngineFrame;
+				}
+
+				// Calc FrameStatus and whether or not it should be pulsing
+				ESimFrameStatus FrameStatus = ESimFrameStatus::Abandoned;
+				uint64 PulseFrameNum = 0;
+												
+				if (SparseData.NetRole != ENP_NetRole::Authority)
+				{
+					if (Tick.TrashedEngineFrame > 0 && Tick.TrashedEngineFrame <= SimulationFrameView.HeadEngineFrame)
+					{
+						FrameStatus = ESimFrameStatus::Trashed;
+						PulseFrameNum = Tick.TrashedEngineFrame;
+					}
+					else if (Tick.ConfirmedEngineFrame > 0 && Tick.ConfirmedEngineFrame <= SimulationFrameView.HeadEngineFrame)
+					{
+						FrameStatus = ESimFrameStatus::Confirmed;
+						PulseFrameNum = Tick.ConfirmedEngineFrame;
+					}
+					else
+					{
+						FrameStatus = Tick.bRepredict ? ESimFrameStatus::Repredicted : ESimFrameStatus::Predicted;
+						PulseFrameNum = Tick.EngineFrame;
+					}							
+				}
+				else
+				{
+					FrameStatus = ESimFrameStatus::Confirmed;
+					PulseFrameNum = Tick.EngineFrame;
+				}
+
+				const bool bPulseFrame = (PulseFrameNum == SimulationFrameView.HeadEngineFrame);
+				const bool bSelected = (&Tick == SelectedTick);
+				const bool bIsSearchHighlighted = PerformSearch(Tick, SimView);
+				
+				SubTrack->TickSourceList.Emplace(FSimulationTrack::FSubTrack::FTickSource{Tick, StartMS, EndMS, bPrevDesat, bPulseFrame, bSelected, bIsSearchHighlighted, FrameStatus});
+				
+				// Process attached NetRecv
+				if (Tick.StartNetRecv)
+				{
+					ensure(Tick.StartNetRecv->NextTick == &Tick);
+					ensure(Tick.StartMS == Tick.StartNetRecv->SimTimeMS);
+					ProcessNetRecv(*Tick.StartNetRecv, SubTrack->NetRecvSourceList, SimView);
+				}
+			}
+			
+			// Look for orphaned NetRecvs (they aren't attached to a simulation tick)
+			for (auto It = SimView.NetRecv.GetIterator(); It; ++It)
+			{
+				const FSimulationData::FNetSerializeRecv& NetRecv = *It;
+				if (NetRecv.NextTick == nullptr)
+				{
+					ProcessNetRecv(NetRecv, Track.OrphanedNetRecvSourceList, SimView);
+				}
+			}
+		}
+	}
+
+	// Now that all tracks and subtracks are laid out on the X axis, calculate the "value Y" for everything
+
+	const TSharedRef<FSlateFontMeasure> FontMeasureService = FSlateApplication::Get().GetRenderer()->GetFontMeasureService();
+	const FSlateFontInfo& Font = FCoreStyle::GetDefaultFontStyle("Regular", 8);
+	const float MaxFontCharHeight = FontMeasureService->Measure(TEXT("!"), Font).Y;
+
+	const float Y_Pad = 20.f;
+	ViewportYMaxValue = 0.f;
+	float Y = ViewportYContentStartValue;
+	for (FSimulationActorGroup& Group : ActorGroups)
+	{
+		const float Y_Start = Y;
+		Y += Y_Pad;
+		Group.ValueY = Y;
+		Y += MaxFontCharHeight;	// Account for Group.DisplayString
+
+		for (FSimulationTrack& Track : Group.SimulationTracks)
+		{
+			Y += Y_Pad/2.f;
+			Track.ValueY = Y;
+
+			for (FSimulationTrack::FSubTrack& SubTrack : Track.SubTracks)
+			{
+				SubTrack.ValueY = Y;
+				SubTrack.H = MaxFontCharHeight;
+				Y += MaxFontCharHeight + Y_Pad;
+			}
+		}
+
+		Group.H = Y - Y_Start;
+	}
+
+	ViewportYMaxValue = Y;
+}
+
+// -----------------------------------------------------------------------------------------------------------
+//	Builds out the final lists of data used to actually draw stuff in OnPaint. This is the only update function
+//	That should be looking at viewport or screenspace information.
+// -----------------------------------------------------------------------------------------------------------
+void SNPSimFrameView::UpdateDraw()
+{
+	// -------------------------------------------------------------------------------------------------------
+	//	Adjust viewport before calculating our min/max draw values
+	// -------------------------------------------------------------------------------------------------------
+
+	FNPAxisViewportInt32& ViewportX = Viewport.GetHorizontalAxisViewport();
+	FNPAxisViewportDouble& ViewportY = Viewport.GetVerticalAxisViewport();
+
+	// This viewport is capable of viewing the entire sim timeline
+	ViewportX.SetMinMaxInterval(ViewportMinSimTimeMS, ViewportMaxSimTimeMS);
+	ViewportY.SetMinMaxValueInterval(0.f, ViewportYMaxValue);
+
+	if (bAutoScrollDirty)
+	{
+		// We want to center on the processed simulation time
+		ViewportX.CenterOnValue(SimulationFrameView.PresentableTimeMS);
+				
+		// But that may put us within the "UpdatePosWithinLimits" limit which causes auto scrolling to happen in the main Tick function.
+		float MinPos, MaxPos;
+		ViewportX.GetScrollLimits(MinPos, MaxPos);
+		ViewportX.EnforceScrollLimits(MinPos, MaxPos, 0); // 0 so that we don't lerp
+				
+		bAutoScrollDirty = false;
+	}
+
+	// -------------------------------------------------------------------------------------------------------
+
+	TArray<FSimulationActorGroup>& ActorGroups = SimulationFrameView.ActorGroups;
+
+	const TSharedRef<FSlateFontMeasure> FontMeasureService = FSlateApplication::Get().GetRenderer()->GetFontMeasureService();
+	const FSlateFontInfo& Font = FCoreStyle::GetDefaultFontStyle("Regular", 8);
+
+	const float SampleW = Viewport.GetSampleWidth();
+	const float ViewHeight = FMath::RoundToFloat(Viewport.GetHeight());
+	const float H = Viewport.GetSubTrackHeight();
+	
+	const float MinWidthForEmbeddedText = FontMeasureService->Measure(TEXT("XXXXX"), Font).X; // Not great, would need to know highest output frame in all tracks which we currently don't cache anywhere
+	const float SubTrackOffsetX = FontMeasureService->Measure(TEXT("   Autonomous  "), Font).X;
+	
+	const float X_Pad = MinWidthForEmbeddedText;
+
+	// Create our draw list used by ::DrawCached
+
+	const int32 MinDrawSimMS = FMath::Max(ViewportMinSimTimeMS, ViewportX.GetValueAtOffset(SubTrackOffsetX));
+	const int32 MaxDrawSimMS = ViewportX.GetValueAtOffset(ViewportX.GetSize());
+
+	for (FSimulationActorGroup& Group : ActorGroups)
+	{
+		Group.Y = ViewportY.GetOffsetForValue(Group.ValueY);
+
+		for (FSimulationTrack& Track : Group.SimulationTracks)
+		{
+			auto ProcessNetRecvList = [&](TArray<FSimulationTrack::FSubTrack::FNetRecvSource>& NetRecvSourceList, TArray<FSimulationTrack::FSubTrack::FDrawNetRecv>& DrawNetRecvList)
+			{
+				for (const FSimulationTrack::FSubTrack::FNetRecvSource& NetRecv : NetRecvSourceList)
+				{
+					if (NetRecv.OffsetNetMS < MinDrawSimMS)
+					{
+						continue;
+					}
+					if (NetRecv.OffsetNetMS > MaxDrawSimMS)
+					{
+						break;
+					}
+
+					const float NetX = ViewportX.GetOffsetForValue(NetRecv.OffsetNetMS);
+					DrawNetRecvList.Emplace( FSimulationTrack::FSubTrack::FDrawNetRecv{NetX, NetRecv} );
+				}
+			};
+
+			Track.Y = ViewportY.GetOffsetForValue(Track.ValueY);
+
+			const FSimulationData::FRestrictedView& SimView = Track.View.Get();
+			const FSimulationData::FSparse& SparseData = SimView.SparseData.Get();
+
+			for (FSimulationTrack::FSubTrack& SubTrack : Track.SubTracks)
+			{
+				SubTrack.DrawTickList.Reset();
+				SubTrack.DrawNetRecvList.Reset();
+
+				SubTrack.Y = ViewportY.GetOffsetForValue(SubTrack.ValueY);
+
+				if (SubTrack.Y < ViewportYContentStartValue || SubTrack.Y > ViewportY.GetMaxPos())
+				{
+					continue;
+				}
+
+				for (const FSimulationTrack::FSubTrack::FTickSource& TickSource : SubTrack.TickSourceList)
+				{
+					if (TickSource.OffsetEndMS < MinDrawSimMS)
+					{
+						continue;
+					}
+					if (TickSource.OffsetStartMS > MaxDrawSimMS)
+					{
+						break;
+					}
+
+					const float StartX = FMath::Max(SubTrackOffsetX, ViewportX.GetOffsetForValue(TickSource.OffsetStartMS));
+					const float EndX = FMath::Max(SubTrackOffsetX, ViewportX.GetOffsetForValue(TickSource.OffsetEndMS));
+
+					const float W = EndX - StartX;
+					if (W <= 0.f)
+					{
+						continue;
+					}
+
+					SubTrack.DrawTickList.Emplace( FSimulationTrack::FSubTrack::FDrawTick{StartX, W, TickSource} );
+				}
+
+				ProcessNetRecvList(SubTrack.NetRecvSourceList, SubTrack.DrawNetRecvList);
+			}
+
+			Track.OrphanedDrawNetRecvList.Reset();
+			ProcessNetRecvList(Track.OrphanedNetRecvSourceList, Track.OrphanedDrawNetRecvList);
+		}
+	}
+}
+
+void SNPSimFrameView::SetAutoScroll(bool bIn)
+{
+	bAutoScroll = bIn;
+	if (bIn)
+	{
+		bAutoScrollDirty = true;
+	}
+}
+void SNPSimFrameView::SetAutoScrollDirty()
+{
+	bAutoScrollDirty = true;
+}
+
+
+////////////////////////////////////////////////////////////////////////////////////////////////////
+
+class FSimFrameTrackDrawContext
+{
+public:
+	explicit FSimFrameTrackDrawContext(const SNPSimFrameView* InSimFrameView, FDrawContext& InDrawContext, const FSimFrameViewDrawHelper& InHelper)
+		: SimFrameView(InSimFrameView)
+		, DrawContext(InDrawContext)
+		, Helper(InHelper)
+	{}
+
+	const FSimFrameViewport& GetViewport() const { return SimFrameView->GetViewport(); }
+	const FVector2D& GetMousePosition() const { return SimFrameView->GetMousePosition(); }
+	
+	/*
+	const TSharedPtr<const ITimingEvent> GetHoveredEvent() const { return TimingView->GetHoveredEvent(); }
+	const TSharedPtr<const ITimingEvent> GetSelectedEvent() const { return TimingView->GetSelectedEvent(); }
+	const TSharedPtr<ITimingEventFilter> GetEventFilter() const { return TimingView->GetEventFilter(); }
+	*/
+	FDrawContext& GetDrawContext() const { return DrawContext; }
+	
+	const FSimFrameViewDrawHelper& GetHelper() const { return Helper; }
+
+public:
+	const SNPSimFrameView* SimFrameView;
+	FDrawContext& DrawContext;
+	const FSimFrameViewDrawHelper& Helper;
+};
+
+////////////////////////////////////////////////////////////////////////////////////////////////////
+
+int32 SNPSimFrameView::OnPaint(const FPaintArgs& Args, const FGeometry& AllottedGeometry, const FSlateRect& MyCullingRect, FSlateWindowElementList& OutDrawElements, int32 LayerId, const FWidgetStyle& InWidgetStyle, bool bParentEnabled) const
+{
+	const bool bEnabled = ShouldBeEnabled(bParentEnabled);
+	const ESlateDrawEffect DrawEffects = bEnabled ? ESlateDrawEffect::None : ESlateDrawEffect::DisabledEffect;
+	FDrawContext DrawContext(AllottedGeometry, MyCullingRect, InWidgetStyle, DrawEffects, OutDrawElements, LayerId);
+
+	const TSharedRef<FSlateFontMeasure> FontMeasureService = FSlateApplication::Get().GetRenderer()->GetFontMeasureService();
+	FSlateFontInfo SummaryFont = FCoreStyle::GetDefaultFontStyle("Regular", 8);
+
+	const FSlateBrush* WhiteBrush = FEditorStyle::GetBrush(TEXT("WhiteBrush")); //FInsightsStyle::Get().GetBrush("WhiteBrush");
+
+	const float ViewWidth = AllottedGeometry.Size.X;
+	const float ViewHeight = AllottedGeometry.Size.Y;
+
+
+	{
+		FSimFrameViewDrawHelper Helper(DrawContext, Viewport);
+		Helper.DrawBackground(SimulationFrameView.PresentableTimeMS);
+
+		Helper.DrawCached(SimulationFrameView, ViewportYContentStartValue);
+
+		DrawHorizontalAxisGrid(DrawContext, WhiteBrush, SummaryFont);
+
+
+
+		if (TooltipOpacity > 0.f)
+		{
+			FString Text;
+			int NumLines = 0;
+
+			if (HoverView.Tick)
+			{
+				Text += FString::Format(TEXT("Simulation Tick:\n"
+													"Simulation Frame {0}\n"
+													"Start Simulation MS: {1}\n"
+													"End Simulation MS: {2}\n"
+													"Delta Simulation MS: {3}\n"
+													"Repredict: {4}\n"
+													"Confirmed Engine Frame: {5}\n"
+													"Trashed Engine Frame: {6}\n"
+													"GFrameNumber: {7}\n\n"),
+					{
+						HoverView.Tick->OutputFrame,
+						HoverView.Tick->StartMS,
+						HoverView.Tick->EndMS,
+						(HoverView.Tick->EndMS - HoverView.Tick->StartMS),
+						HoverView.Tick->bRepredict,
+						HoverView.Tick->ConfirmedEngineFrame,
+						HoverView.Tick->TrashedEngineFrame,
+						HoverView.Tick->EngineFrame
+					});
+
+				NumLines += 10;
+			}
+
+			if (HoverView.NetRecv)
+			{
+				
+
+				Text += FString::Format(TEXT("Net Receive:\n"
+													"SimFrame: {0}\n"
+													"SimTime MS: {1}\n"
+													"Status: {2}\n"
+													"GFrameNumber: {3}\n"),
+					{
+						HoverView.NetRecv->Frame,
+						HoverView.NetRecv->SimTimeMS,
+						LexToString(HoverView.NetRecv->Status),						
+						HoverView.NetRecv->EngineFrame
+					});
+
+				NumLines += 5;
+			}
+
+			
+
+			FVector2D TextSize = FontMeasureService->Measure(Text, SummaryFont);
+
+			const float DX = 2.0f;
+			const float W2 = TextSize.X / 2 + DX;
+
+			const FNPAxisViewportInt32& ViewportX = Viewport.GetHorizontalAxisViewport();
+			float X1 = HoverView.X;
+			float CX = X1 + FMath::RoundToFloat(Viewport.GetSampleWidth() / 2);
+			
+			if (CX + W2 > ViewportX.GetSize())
+			{
+				CX = FMath::RoundToFloat(ViewportX.GetSize() - W2);
+			}
+			if (CX - W2 < 0)
+			{
+				CX = W2;
+			}
+			
+			const float H = 2.0f + 13.0f * NumLines;
+			const float Y = FMath::Max(HoverView.Y - H, 20.f);
+
+
+			if (MousePosition.Y > Y && MousePosition.Y < (Y + H))
+			{
+				CX -= (W2);
+				if ((CX-W2) < 0.f)
+				{
+					CX += (W2*3.f);//+ 10.f;
+				}
+			}
+
+
+
+			DrawContext.DrawBox(CX - W2, Y, 2 * W2, H, WhiteBrush, FLinearColor(0.7, 0.7, 0.7, TooltipOpacity));
+			DrawContext.LayerId++;
+			DrawContext.DrawText(CX - W2 + DX, Y + 1.0f, Text, SummaryFont, FLinearColor(0.0, 0.0, 0.0, TooltipOpacity));
+			DrawContext.LayerId++;
+		}
+	}
+
+	return SCompoundWidget::OnPaint(Args, AllottedGeometry, MyCullingRect, OutDrawElements, LayerId, InWidgetStyle, bParentEnabled && IsEnabled());
+}
+
+void SNPSimFrameView::DrawHorizontalAxisGrid(FDrawContext& DrawContext, const FSlateBrush* Brush, const FSlateFontInfo& Font) const
+{
+	const FNPAxisViewportInt32& ViewportX = Viewport.GetHorizontalAxisViewport();
+
+	const float RoundedViewWidth = FMath::RoundToFloat(ViewportX.GetSize());
+
+	constexpr float MinDX = 100.0f; // min horizontal distance between vertical grid lines
+
+	const int32 LeftIndex = ViewportX.GetValueAtOffset(0.0f);
+	const int32 GridIndex = ViewportX.GetValueAtOffset(MinDX);
+	const int32 RightIndex = ViewportX.GetValueAtOffset(RoundedViewWidth);
+	const int32 Delta = GridIndex - LeftIndex;
+
+	if (Delta > 0)
+	{
+		// Compute rounding based on magnitude of visible range of samples (Delta).
+		int32 Power10 = 1;
+		int32 Delta10 = Delta;
+		while (Delta10 > 0)
+		{
+			Delta10 /= 10;
+			Power10 *= 10;
+		}
+		if (Power10 >= 100)
+		{
+			Power10 /= 100;
+		}
+		else
+		{
+			Power10 = 1;
+		}
+
+		const int32 Grid = ((Delta + Power10 - 1) / Power10) * Power10; // next value divisible with a multiple of 10
+
+		// Skip grid lines for negative indices.
+		double StartIndex = ((LeftIndex + Grid - 1) / Grid) * Grid;
+		while (StartIndex < 0)
+		{
+			StartIndex += Grid;
+		}
+
+		const float ViewHeight = Viewport.GetHeight();
+
+		const FLinearColor GridColor(0.0f, 0.0f, 0.0f, 0.1f);
+		//const FLinearColor TextBgColor(0.05f, 0.05f, 0.05f, 1.0f);
+		//const FLinearColor TextColor(1.0f, 1.0f, 1.0f, 1.0f);
+		const FLinearColor TopTextColor(1.0f, 1.0f, 1.0f, 0.7f);
+
+		//const TSharedRef<FSlateFontMeasure> FontMeasureService = FSlateApplication::Get().GetRenderer()->GetFontMeasureService();
+
+		for (int32 Index = StartIndex; Index < RightIndex; Index += Grid)
+		{
+			const float X = FMath::RoundToFloat(ViewportX.GetOffsetForValue(Index));
+
+			// Draw vertical grid line.
+			DrawContext.DrawBox(X, 0, 1, ViewHeight, Brush, GridColor);
+
+			const FString Text = FText::AsNumber(Index).ToString();
+			//const FVector2D TextSize = FontMeasureService->Measure(Text, Font);
+			//constexpr float TextH = 14.0f;
+
+			// Draw background for index text.
+			//DrawContext.DrawBox(X, ViewHeight - TextH, TextSize.X + 4.0f, TextH, Brush, TextBgColor);
+
+			// Draw index text.
+			//DrawContext.DrawText(X + 2.0f, ViewHeight - TextH + 1.0f, Text, Font, TextColor);
+
+			DrawContext.DrawText(X + 2.0f, 10.0f, Text, Font, TopTextColor);
+		}
+		DrawContext.LayerId++;
+	}
+}
+
+////////////////////////////////////////////////////////////////////////////////////////////////////
+
+
+FReply SNPSimFrameView::OnMouseButtonDown(const FGeometry& MyGeometry, const FPointerEvent& MouseEvent)
+{
+	FReply Reply = FReply::Unhandled();
+
+	MousePositionOnButtonDown = MyGeometry.AbsoluteToLocal(MouseEvent.GetScreenSpacePosition());
+	ViewportPosXOnButtonDown = Viewport.GetHorizontalAxisViewport().GetPos();
+	ViewportPosYOnButtonDown = Viewport.GetVerticalAxisViewport().GetPos();
+
+	if (MouseEvent.GetEffectingButton() == EKeys::LeftMouseButton)
+	{
+		bIsLMB_Pressed = true;
+
+		// Capture mouse.
+		Reply = FReply::Handled().CaptureMouse(SharedThis(this));
+
+		FMouseOverInfo MouseOver = GetMouseOverInfo(MousePositionOnButtonDown);
+
+		if (MouseOver.DrawSimTick || MouseOver.DrawNetRecv)
+		{
+			const FSimulationData::FTick* ClickedTick = MouseOver.DrawSimTick ? &MouseOver.DrawSimTick->Source.Tick : nullptr;
+			const FSimulationData::FNetSerializeRecv* ClickedMouseRecv = MouseOver.GetNetRecv(NetworkPredictionWindow->GetFilteredDataCollection().LastEngineFrame);
+			NetworkPredictionWindow->NotifySimContentClicked(FSimContentsView { MouseOver.Track->View, 
+																				ClickedTick, 
+																				ClickedMouseRecv });
+			bIsStateDirty = true;
+		}
+	}
+	else if (MouseEvent.GetEffectingButton() == EKeys::RightMouseButton)
+	{
+		bIsRMB_Pressed = true;
+
+		// Capture mouse, so we can scroll outside this widget.
+		Reply = FReply::Handled().CaptureMouse(SharedThis(this));
+	}
+	else if (MouseEvent.GetEffectingButton() == EKeys::ThumbMouseButton)
+	{
+		NetworkPredictionWindow->JumpPreviousViewedEngineFrame();
+		Reply = FReply::Handled();
+	}
+	else if (MouseEvent.GetEffectingButton() == EKeys::ThumbMouseButton2)
+	{
+		NetworkPredictionWindow->JumpNextViewedEngineFrame();
+		Reply = FReply::Handled();
+	}
+	return Reply;
+}
+
+////////////////////////////////////////////////////////////////////////////////////////////////////
+
+FReply SNPSimFrameView::OnMouseButtonUp(const FGeometry& MyGeometry, const FPointerEvent& MouseEvent)
+{
+	FReply Reply = FReply::Unhandled();
+
+	MousePositionOnButtonUp = MyGeometry.AbsoluteToLocal(MouseEvent.GetScreenSpacePosition());
+
+	const bool bIsValidForMouseClick = MousePositionOnButtonUp.Equals(MousePositionOnButtonDown, MOUSE_SNAP_DISTANCE);
+
+	if (MouseEvent.GetEffectingButton() == EKeys::LeftMouseButton)
+	{
+		if (bIsLMB_Pressed)
+		{
+			if (bIsScrolling)
+			{
+				bIsScrolling = false;
+				CursorType = ECursorType::Default;
+			}
+			else if (bIsValidForMouseClick)
+			{
+				//SelectSampleAtMousePosition(MousePositionOnButtonUp.X, MousePositionOnButtonUp.Y, MouseEvent);
+			}
+
+			bIsLMB_Pressed = false;
+
+			// Release the mouse.
+			Reply = FReply::Handled().ReleaseMouseCapture();
+		}
+	}
+	else if (MouseEvent.GetEffectingButton() == EKeys::RightMouseButton)
+	{
+		if (bIsRMB_Pressed)
+		{
+			if (bIsScrolling)
+			{
+				bIsScrolling = false;
+				CursorType = ECursorType::Default;
+			}
+			else if (bIsValidForMouseClick)
+			{
+				FMouseOverInfo MouseOver = GetMouseOverInfo(MousePositionOnButtonDown);
+
+				FMenuBuilder MenuBuilder(true, nullptr);
+				const FSimulationData::FNetSerializeRecv* NetRecv = MouseOver.GetNetRecv(NetworkPredictionWindow->GetFilteredDataCollection().LastEngineFrame);
+
+				if (MouseOver.Group)
+				{
+					MenuBuilder.AddMenuEntry(
+						LOCTEXT("SortGroupToTop", "Sort to top"), 
+						LOCTEXT("SortGroupToTopTip", "Bring this Actor Group to the top of the Simulation Timeline View"),
+						FSlateIcon(),
+						FUIAction(FExecuteAction::CreateSP(this, &SNPSimFrameView::SortActorGroupToTop, MouseOver.Group)),
+						NAME_None,
+						EUserInterfaceActionType::Button
+					);
+				}
+				
+				if (MouseOver.DrawSimTick)
+				{
+					FString Text = FString::Printf(TEXT("Sim Tick: Jump To %d"), MouseOver.DrawSimTick->Source.Tick.EngineFrame);
+
+					MenuBuilder.AddMenuEntry(
+						FText::FromString(Text), 
+						LOCTEXT("JumpToSimFrameEngineNumToolTip", "Jump to the Engine Frame that this tick was made on."),
+						FSlateIcon(),
+						FUIAction(FExecuteAction::CreateSP(NetworkPredictionWindow.Get(), &SNPWindow::SetEngineFrame, MouseOver.DrawSimTick->Source.Tick.EngineFrame, true)),
+						NAME_None,
+						EUserInterfaceActionType::Button
+					);
+				}
+					
+				if (NetRecv)
+				{
+					FString Text = FString::Printf(TEXT("Net Recv: Jump To %d"), NetRecv->EngineFrame);
+
+					MenuBuilder.AddMenuEntry(
+						FText::FromString(Text),
+						LOCTEXT("JumpToNetRecvEngineNumToolTip", "Jump to the Engine Frame that this Net Receive was received on."),
+						FSlateIcon(),
+						FUIAction(FExecuteAction::CreateSP(NetworkPredictionWindow.Get(), &SNPWindow::SetEngineFrame, NetRecv->EngineFrame, true)),
+						NAME_None,
+						EUserInterfaceActionType::Button
+					);
+				}
+
+				TSharedRef<SWidget> MenuWidget = MenuBuilder.MakeWidget();
+
+				FWidgetPath EventPath = MouseEvent.GetEventPath() != nullptr ? *MouseEvent.GetEventPath() : FWidgetPath();
+				const FVector2D ScreenSpacePosition = MouseEvent.GetScreenSpacePosition();
+				FSlateApplication::Get().PushMenu(SharedThis(this), EventPath, MenuWidget, ScreenSpacePosition, FPopupTransitionEffect::ContextMenu);
+			}
+
+			bIsRMB_Pressed = false;
+
+			// Release mouse as we no longer scroll.
+			Reply = FReply::Handled().ReleaseMouseCapture();
+		}
+	}
+
+	return Reply;
+}
+
+void SNPSimFrameView::OnMouseEnter(const FGeometry& MyGeometry, const FPointerEvent& MouseEvent)
+{
+
+}
+void SNPSimFrameView::OnMouseLeave(const FPointerEvent& MouseEvent)
+{
+	if (!HasMouseCapture())
+	{
+		bIsLMB_Pressed = false;
+		bIsRMB_Pressed = false;
+
+		HoverView.Reset();
+		TooltipOpacity = 0.f;
+
+		CursorType = ECursorType::Default;
+	}
+}
+
+FReply SNPSimFrameView::OnMouseWheel(const FGeometry& MyGeometry, const FPointerEvent& MouseEvent)
+{
+	MousePosition = MyGeometry.AbsoluteToLocal(MouseEvent.GetScreenSpacePosition());
+
+	if (MouseEvent.GetModifierKeys().IsShiftDown())
+	{
+		FNPAxisViewportDouble& ViewportY = Viewport.GetVerticalAxisViewport();
+
+		// Zoom in/out vertically.
+		const float Delta = MouseEvent.GetWheelDelta();
+		constexpr float ZoomStep = 0.25f; // as percent
+		float ScaleY;
+
+		if (Delta > 0)
+		{
+			ScaleY = ViewportY.GetScale() * FMath::Pow(1.0f + ZoomStep, Delta);
+		}
+		else
+		{
+			ScaleY = ViewportY.GetScale() * FMath::Pow(1.0f / (1.0f + ZoomStep), -Delta);
+		}
+
+		ViewportY.SetScale(ScaleY);
+		//UpdateVerticalScrollBar();
+	}
+	else //if (MouseEvent.GetModifierKeys().IsControlDown())
+	{
+		// Zoom in/out horizontally.
+		const float Delta = MouseEvent.GetWheelDelta();
+		ZoomHorizontally(Delta, MousePosition.X);
+	}
+
+	return FReply::Handled();
+}
+
+FReply SNPSimFrameView::OnMouseMove(const FGeometry& MyGeometry, const FPointerEvent& MouseEvent)
+{
+	FReply Reply = FReply::Unhandled();
+
+	MousePosition = MyGeometry.AbsoluteToLocal(MouseEvent.GetScreenSpacePosition());
+
+	if (!MouseEvent.GetCursorDelta().IsZero())
+	{
+		if (MouseEvent.IsMouseButtonDown(EKeys::LeftMouseButton) ||
+			MouseEvent.IsMouseButtonDown(EKeys::RightMouseButton))
+		{
+			if (HasMouseCapture())
+			{
+				if (!bIsScrolling)
+				{
+					bIsScrolling = true;
+					CursorType = ECursorType::Hand;
+					
+					HoverView.Reset();
+				}
+
+				FNPAxisViewportInt32& ViewportX = Viewport.GetHorizontalAxisViewport();
+				const float PosX = ViewportPosXOnButtonDown + (MousePositionOnButtonDown.X - MousePosition.X);
+				ViewportX.ScrollAtValue(ViewportX.GetValueAtPos(PosX)); // align viewport position with sample
+				UpdateHorizontalScrollBar();
+
+
+				FNPAxisViewportDouble& ViewportY = Viewport.GetVerticalAxisViewport();
+				const float PosY = ViewportPosYOnButtonDown + (MousePositionOnButtonDown.Y - MousePosition.Y);
+				ViewportY.ScrollAtValue(ViewportY.GetValueAtPos(PosY));
+				UpdateVerticalScrollBar();
+
+
+				OnUserScroll();
+			}
+		}
+		else
+		{
+			TooltipOpacity = 0.f;
+			HoverView.Tick = nullptr;
+			HoverView.NetRecv = nullptr;
+
+			FMouseOverInfo MouseOver = GetMouseOverInfo(MousePosition);
+			if (MouseOver.DrawNetRecv && ensure(MouseOver.Track))
+			{
+				HoverView.NetRecv = &MouseOver.DrawNetRecv->Source.NetRecv;
+				HoverView.X = MouseOver.DrawNetRecv->X;
+				HoverView.Y = MouseOver.Track->Y;
+				TooltipOpacity = 1.f;
+			}
+
+			if (MouseOver.DrawSimTick && ensure(MouseOver.Track))
+			{
+				HoverView.Tick = &MouseOver.DrawSimTick->Source.Tick;
+				HoverView.X = MouseOver.DrawSimTick->X;
+				HoverView.Y = MouseOver.Track->Y;
+				TooltipOpacity = 1.f;
+			}
+		}
+
+		Reply = FReply::Handled();
+	}
+
+	return Reply;
+}
+
+SNPSimFrameView::FMouseOverInfo SNPSimFrameView::GetMouseOverInfo(FVector2D InMousePosition) const
+{
+	FMouseOverInfo Info;
+
+	const float Y = InMousePosition.Y;
+	const float X = InMousePosition.X;
+
+	if (Y > Viewport.GetContentStartY())
+	{
+		auto GetMouseOverInfo = [&]()
+		{
+			for (const FSimulationActorGroup& Group : SimulationFrameView.ActorGroups)
+			{
+				if (Y >= Group.Y && Y < (Group.Y + Group.H))
+				{
+					Info.Group = &Group;
+					for (const FSimulationTrack& SimTrack : Group.SimulationTracks)
+					{
+						auto FindNetRecv = [&](const TArray<FSimulationTrack::FSubTrack::FDrawNetRecv>& DrawNetRecvList)
+						{
+							for (const FSimulationTrack::FSubTrack::FDrawNetRecv& DrawNetRecv : DrawNetRecvList)
+							{
+								if (X >= DrawNetRecv.X && X < (DrawNetRecv.X + DrawNetRecv.W()))
+								{
+									Info.DrawNetRecv = &DrawNetRecv;
+									return;
+								}
+							}
+						};
+
+						Info.Track = &SimTrack;
+						for (const FSimulationTrack::FSubTrack& SubTrack : SimTrack.SubTracks)
+						{
+							if (Y >= SubTrack.Y && Y < (SubTrack.Y + SubTrack.H))
+							{
+								Info.SubTrack = &SubTrack;
+
+								auto FindTick = [&]()
+								{
+									for (const FSimulationTrack::FSubTrack::FDrawTick& DrawTick : SubTrack.DrawTickList)
+									{
+										if (X >= DrawTick.X && X < (DrawTick.X + DrawTick.W))
+										{
+											Info.DrawSimTick = &DrawTick;
+											return;
+										}
+									}
+								};
+
+								if (&SubTrack == &SimTrack.SubTracks[0])
+								{
+									FindNetRecv(SimTrack.OrphanedDrawNetRecvList);
+								}
+
+								FindNetRecv(SubTrack.DrawNetRecvList);
+								FindTick();
+								return;
+							}
+						}
+
+					}
+					return;
+				}
+			}
+		};
+				
+		GetMouseOverInfo();
+	}
+
+	return Info;
+}
+
+void SNPSimFrameView::Reset()
+{
+	Viewport.Reset();
+	FNPAxisViewportInt32& ViewportX = Viewport.GetHorizontalAxisViewport();
+	ViewportX.SetScaleLimits(1.f, 16.0f); // 1 [sample/px] to 16 [px/sample]
+	ViewportX.SetScale(5.0f);
+	FNPAxisViewportDouble& ViewportY = Viewport.GetVerticalAxisViewport();
+	bIsViewportDirty = true;
+	bIsStateDirty = true;
+}
+
+void SNPSimFrameView::HorizontalScrollBar_OnUserScrolled(float ScrollOffset)
+{
+	FNPAxisViewportInt32& ViewportX = Viewport.GetHorizontalAxisViewport();
+	ViewportX.OnUserScrolled(HorizontalScrollBar, ScrollOffset);
+	OnUserScroll();
+}
+
+void SNPSimFrameView::UpdateHorizontalScrollBar()
+{
+	FNPAxisViewportInt32& ViewportX = Viewport.GetHorizontalAxisViewport();
+	ViewportX.UpdateScrollBar(HorizontalScrollBar);
+}
+
+void SNPSimFrameView::VerticalScrollBar_OnUserScrolled(float ScrollOffset)
+{
+	FNPAxisViewportDouble& ViewportY = Viewport.GetVerticalAxisViewport();
+	ViewportY.OnUserScrolled(VerticalScrollBar, ScrollOffset);
+	OnUserScroll();
+}
+
+void SNPSimFrameView::UpdateVerticalScrollBar()
+{
+	FNPAxisViewportDouble& ViewportY = Viewport.GetVerticalAxisViewport();
+	ViewportY.UpdateScrollBar(VerticalScrollBar);
+}
+
+void SNPSimFrameView::ZoomHorizontally(const float Delta, const float X)
+{
+	FNPAxisViewportInt32& ViewportX = Viewport.GetHorizontalAxisViewport();
+	ViewportX.RelativeZoomWithFixedOffset(Delta, X);
+	ViewportX.ScrollAtValue(ViewportX.GetValueAtPos(ViewportX.GetPos())); // align viewport position with sample
+	UpdateHorizontalScrollBar();
+	bIsDrawDirty = true;
+}
+
+void SNPSimFrameView::OnUserScroll()
+{
+	bAutoScroll = false;
+	bIsDrawDirty = true;
+}
+
+void SNPSimFrameView::OnGetOptionsMenu(FMenuBuilder& Builder)
+{
+	Builder.BeginSection(NAME_None, LOCTEXT("SimulationFrameViewLabel", "Simulation Frame View"));
+	{
+		Builder.AddMenuEntry(
+			LOCTEXT("CompactViewLabel","Compact View"), 
+			LOCTEXT("CompactViewLabelToolTip", "Display Simulation Frames in a more compact fashion. Takes less vertical space but can make harder to read."),
+			FSlateIcon(),
+			FUIAction(FExecuteAction::CreateSP(this, &SNPSimFrameView::ToggleCompactSimFrameView),
+					  FCanExecuteAction(),
+					  FIsActionChecked::CreateSP(this, &SNPSimFrameView::CompactSimFrameView)),
+			NAME_None,
+			EUserInterfaceActionType::ToggleButton
+		);
+
+		Builder.AddMenuEntry(
+			LOCTEXT("LinearViewLabel","Linear View"), 
+			LOCTEXT("LinearViewLabelToolTip", "Display Simulation Frames in contingous linear series when possible. Can take more vertical space but is generally easier to read."),
+			FSlateIcon(),
+			FUIAction(FExecuteAction::CreateSP(this, &SNPSimFrameView::ToggleLinearSimFrameview),
+					  FCanExecuteAction(),
+					  FIsActionChecked::CreateSP(this, &SNPSimFrameView::LinearSimFrameView)),
+			NAME_None,
+			EUserInterfaceActionType::ToggleButton
+		);
+	}
+	Builder.EndSection();
+}
+
+void SNPSimFrameView::SearchUserData(const FText& InFilterText)
+{
+	UserStateSearchString = InFilterText.ToString();
+	bIsStateDirty = true;
+}
+
+bool SNPSimFrameView::PerformSearch(const FSimulationData::FTick& Tick, const FSimulationData::FRestrictedView& SimView)
+{
+	if (UserStateSearchString.IsEmpty())
+	{
+		return false;
+	}
+
+	TArray<const FSimulationData::FUserState*, TInlineAllocator<5>> UserStates;
+
+	const int32 InputFrame = Tick.OutputFrame - 1;
+	const int32 OutputFrame = Tick.OutputFrame;
+	const uint64 EngineFrame = Tick.EngineFrame;
+
+	const uint8 Mask = (uint8)ENP_UserStateSource::NetRecv;
+
+	UserStates.Add(SimView.UserData.Get(ENP_UserState::Input, InputFrame, EngineFrame, Mask));
+
+	UserStates.Add(SimView.UserData.Get(ENP_UserState::Sync, InputFrame, EngineFrame, Mask));
+	UserStates.Add(SimView.UserData.Get(ENP_UserState::Aux, InputFrame, EngineFrame, Mask));
+
+	UserStates.Add(SimView.UserData.Get(ENP_UserState::Sync, OutputFrame, EngineFrame, Mask));
+	UserStates.Add(SimView.UserData.Get(ENP_UserState::Aux, OutputFrame, EngineFrame, Mask));
+	
+	return PerformSearchInternal(UserStates);
+}
+
+bool SNPSimFrameView::PerformSearch(const FSimulationData::FNetSerializeRecv& NetRecv, const FSimulationData::FRestrictedView& SimView)
+{
+	if (UserStateSearchString.IsEmpty())
+	{
+		return false;
+	}
+
+	TArray<const FSimulationData::FUserState*, TInlineAllocator<3>> UserStates;
+
+	const int32 SimFrame = NetRecv.Frame;
+	const uint64 EngineFrame = NetRecv.EngineFrame;
+
+	uint8 Mask = ~((uint8)ENP_UserStateSource::NetRecv | (uint8)ENP_UserStateSource::NetRecvCommit);
+	UserStates.Add(SimView.UserData.Get(ENP_UserState::Input, SimFrame, EngineFrame, Mask));
+	UserStates.Add(SimView.UserData.Get(ENP_UserState::Sync, SimFrame, EngineFrame, Mask));
+	UserStates.Add(SimView.UserData.Get(ENP_UserState::Aux, SimFrame, EngineFrame, Mask));
+
+	return PerformSearchInternal(UserStates);
+}
+
+bool SNPSimFrameView::PerformSearchInternal(const TArrayView<const FSimulationData::FUserState* const>& UserStates)
+{
+	for (const FSimulationData::FUserState* UserState : UserStates)
+	{
+		if (UserState && FCString::Stristr(UserState->UserStr, *UserStateSearchString) != nullptr)
+		{
+			return true;
+		}
+	}
+	return false;
+}
+
+void SNPSimFrameView::SortActorGroupToTop(const FSimulationActorGroup* Group)
+{
+	UserSortedNetActors.Remove(Group->ID);
+	UserSortedNetActors.Push(Group->ID);
+	bIsStateDirty = true;
+}
+
+#undef LOCTEXT_NAMESPACE
\ No newline at end of file
diff --git a/Engine/Plugins/Runtime/NetworkPredictionInsights/Source/NetworkPredictionInsights/Private/UI/SNPSimFrameView.h b/Engine/Plugins/Runtime/NetworkPredictionInsights/Source/NetworkPredictionInsights/Private/UI/SNPSimFrameView.h
new file mode 100644
--- /dev/null
+++ b/Engine/Plugins/Runtime/NetworkPredictionInsights/Source/NetworkPredictionInsights/Private/UI/SNPSimFrameView.h
@@ -0,0 +1,442 @@
+// Copyright Epic Games, Inc. All Rights Reserved.
+
+#pragma once
+
+#include "NPAxisViewportInt32.h"
+#include "NPAxisViewportDouble.h"
+
+enum class ESimFrameStatus : uint8
+{
+	Predicted,	// A predicted frame that has not been rolled back or confirmed in any way
+	Repredicted,// Sim was rolled back to a point prior to this frame, and we repredicted this frame to *catch back up* to where we were
+	Confirmed,	// We received authoratative state > that this frame and didn't rollback/reconcile
+	Trashed,	// Frame that was thrown out: we rolled the sim back and ticked Repredicted frames past this one.
+	Abandoned,	// The sim was skipped ahead by receiving authoritative state: this frame was never used as input into the next frame
+};
+
+// ---------------------------------------------------------------------------------------------------------------------------------
+//
+//	[ActorName, Role, NetGUID, Simulation Group/Type]											} FSimulationActorGroup[0]
+//	--------------------------------------------------------------------------------------------------------------------------------
+//	  | Auto Proxy |	******************PP													} SubTrack[0]  \ SimulationTracks[0]
+//	  |	           |                !!!!!!														} SubTrack[1]  /
+//	  ------------------------------------------------------------------------------------------------------------------------------
+//	  | Authority  |    *******************														} SubTrack[0] -> SimulationTracks[1]
+//	  ------------------------------------------------------------------------------------------------------------------------------
+//	  | Sim Proxy  |    *******************														} SubTrack[0] -> SimulationTracks[2]
+// ---------------------------------------------------------------------------------------------------------------------------------
+
+struct FSimulationTrack
+{
+	// Processed Data about subtracks: this will indicate overlapping simulation steps
+	// essentially, how we increase the Y size of the track in the UI
+	struct FSubTrack
+	{
+		struct FTickSource
+		{
+			const FSimulationData::FTick& Tick;
+			const FSimTime OffsetStartMS;
+			const FSimTime OffsetEndMS;
+			const bool bDesaturate;
+			const bool bPulse;
+			const bool bSelected;
+			const bool bSearchHighlighted;
+			const ESimFrameStatus Status;
+		};
+
+		struct FNetRecvSource
+		{
+			const FSimulationData::FNetSerializeRecv& NetRecv;
+			const FSimTime OffsetNetMS;
+			const bool bPulse;
+			const bool bSelected;
+			const bool bSearchHighlighted;
+		};
+
+		struct FDrawTick
+		{
+			const float X;
+			const float W;
+			const FTickSource& Source;
+		};
+
+		struct FDrawNetRecv
+		{
+			static constexpr float W() { return 3.f; }
+
+			float X;
+			const FNetRecvSource& Source;
+		};
+
+		// All ticks that belong to this track. These are not culled wrt the viewport
+		TArray<FTickSource> TickSourceList;
+		TArray<FNetRecvSource> NetRecvSourceList;
+		int32 PrevMS = TNumericLimits<int32>::Min();
+
+		// These are the actual draw lists that get flung over to ::DrawCached
+		TArray<FDrawTick> DrawTickList;
+		TArray<FDrawNetRecv> DrawNetRecvList;
+		
+		float ValueY=0.f; // untransformed by viewport/scrolling value
+		float Y=0.f;
+		float H=0.f;
+	};
+
+	TSharedRef<FSimulationData::FRestrictedView> View;
+
+	float Y;
+	float ValueY; // untransformed by viewport/scrolling value
+	FString DisplayString;
+	TArray<FSubTrack> SubTracks;
+
+	// Orphaned Recvs don't map to a specific subtrack.
+	TArray<FSubTrack::FNetRecvSource> OrphanedNetRecvSourceList;
+	TArray<FSubTrack::FDrawNetRecv> OrphanedDrawNetRecvList;
+};
+
+struct FSimulationActorGroup
+{
+	FSimNetActorID ID;
+
+	FString DebugName;
+	FName GroupName;
+	bool bHasAutoProxy; // Is there an auto proxy sim in this group
+
+	uint64 MaxEngineFrame = 0; // Engine frame number of last simulation frame amonst all simulations for this actor
+	FSimTime MaxAllowedSimTime = 0;
+	
+	FSimTime OffsetSimTimeMS; // How much to offset this group when drawing anything on the sim timeline
+
+	TArray<FSimulationTrack> SimulationTracks;
+
+	// Rendering Data. Only filled out if being drawn
+	FString DisplayString;
+	float Y;
+	float H;
+
+	float ValueY; // untransformed by viewport/scrolling value
+};
+
+// The overall processed data that the Simulation Frame Widget uses to draw
+struct FSimulationFrameView
+{
+	void Reset()
+	{
+		HeadEngineFrame = 0;
+		ActorGroups.Reset();
+	}
+
+	uint64 HeadEngineFrame = 0;
+	FSimTime PresentableTimeMS = 0; // The "current" time that everything is anchored around
+	TArray<FSimulationActorGroup> ActorGroups;
+};
+
+// Cached data about what we are hovering over
+struct FSimFrameHoverView
+{
+	const FSimulationData::FTick* Tick = nullptr;
+	const FSimulationData::FNetSerializeRecv* NetRecv = nullptr;
+	FSimTime X = 0;
+	FSimTime Y = 0;
+
+	void Reset()
+	{
+		Tick = nullptr;
+		NetRecv = nullptr;
+		X = 0;
+		Y = 0;
+	}
+};
+
+////////////////////////////////////////////////////////////////////////////////////////////////////
+
+// Viewport that is simulation frame based (not real time!)
+class FSimFrameViewport
+{
+private:
+	static constexpr float SLATE_UNITS_TOLERANCE = 0.1f;
+
+public:
+	FSimFrameViewport()
+	{
+		Reset();
+	}
+
+	void Reset()
+	{
+		HorizontalAxisViewport.Reset();
+		VerticalAxisViewport.Reset();
+	}
+
+	const FNPAxisViewportInt32& GetHorizontalAxisViewport() const { return HorizontalAxisViewport; }
+	FNPAxisViewportInt32& GetHorizontalAxisViewport() { return HorizontalAxisViewport; }
+
+	const FNPAxisViewportDouble& GetVerticalAxisViewport() const { return VerticalAxisViewport; }
+	FNPAxisViewportDouble& GetVerticalAxisViewport() { return VerticalAxisViewport; }
+
+	float GetWidth() const { return HorizontalAxisViewport.GetSize(); }
+	float GetHeight() const { return VerticalAxisViewport.GetSize(); }
+
+	bool SetSize(const float InWidth, const float InHeight)
+	{
+		const bool bWidthChanged = HorizontalAxisViewport.SetSize(InWidth);
+		const bool bHeightChanged = VerticalAxisViewport.SetSize(InHeight);
+		if (bWidthChanged || bHeightChanged)
+		{
+			OnSizeChanged();
+			return true;
+		}
+		return false;
+	}
+
+	float GetSampleWidth() const { return HorizontalAxisViewport.GetSampleSize(); }
+	int32 GetNumFramesPerSample() const { return HorizontalAxisViewport.GetNumSamplesPerPixel(); }
+	int32 GetFirstFrameIndex() const { return HorizontalAxisViewport.GetValueAtOffset(0.0f); }
+
+	float GetContentStartY() const { return 30.f; } //GetVerticalAxisViewport().GetSize() - 1.f; }
+	float GetTrackPadding() const { return GetSubTrackHeight(); }
+	float GetSubTrackHeight() const
+	{
+		const float BaselineY = FMath::RoundToFloat(VerticalAxisViewport.GetOffsetForValue(0.0));
+		const float ValueY = FMath::RoundToFloat(VerticalAxisViewport.GetOffsetForValue(20.0));
+		const float H = ValueY - BaselineY;
+		return H;
+	}
+	
+
+private:
+	void OnSizeChanged()
+	{
+	}
+
+private:
+	FNPAxisViewportInt32 HorizontalAxisViewport;
+	FNPAxisViewportDouble VerticalAxisViewport;
+};
+
+////////////////////////////////////////////////////////////////////////////////////////////////////
+
+class FSimFrameViewDrawHelper
+{
+public:
+	enum class EHighlightMode : uint32
+	{
+		Hovered = 1,
+		Selected = 2,
+		SelectedAndHovered = 3
+	};
+
+public:
+	explicit FSimFrameViewDrawHelper(const FDrawContext& InDrawContext, const FSimFrameViewport& InViewport);
+
+	/**
+	 * Non-copyable
+	 */
+	FSimFrameViewDrawHelper(const FSimFrameViewDrawHelper&) = delete;
+	FSimFrameViewDrawHelper& operator=(const FSimFrameViewDrawHelper&) = delete;
+
+	void DrawBackground(FSimTime PresentableLine) const;
+	void DrawCached(const FSimulationFrameView& View, float MinY) const;
+	
+	FLinearColor GetFramColorByStatus(ESimFrameStatus Status, const bool bDesaturate) const;
+	FLinearColor GetRecvColorByStatus(ENetSerializeRecvStatus Status) const;
+	
+	/*
+	void DrawSampleHighlight(const FNetworkPacketAggregatedSample& Sample, EHighlightMode Mode) const;
+	void DrawSelection(int32 StartPacketIndex, int32 EndPacketIndex) const;
+	
+
+	static FLinearColor GetColorByStatus(Trace::ENetProfilerDeliveryStatus Status);
+
+	int32 GetNumPackets() const { return NumPackets; }
+	int32 GetNumDrawSamples() const { return NumDrawSamples; }
+	*/
+
+private:
+	const FDrawContext& DrawContext;
+	const FSimFrameViewport& Viewport;
+
+	const FSlateBrush* WhiteBrush;
+	const FSlateBrush* HoveredEventBorderBrush;
+	const FSlateBrush* SelectedEventBorderBrush;
+	const FSlateFontInfo SelectionFont;
+
+	// Debug stats.
+	mutable int32 NumFrames = 0;
+	mutable int32 NumDrawSamples = 0;
+};
+
+
+////////////////////////////////////////////////////////////////////////////////////////////////////
+
+class SNPSimFrameView : public SCompoundWidget
+{
+public:
+
+	/** Number of pixels. */
+	static constexpr float MOUSE_SNAP_DISTANCE = 2.0f;
+
+	enum class ECursorType
+	{
+		Default,
+		Arrow,
+		Hand,
+	};
+
+
+	SNPSimFrameView();
+	virtual ~SNPSimFrameView();
+
+	SLATE_BEGIN_ARGS(SNPSimFrameView)
+	{
+		_Clipping = EWidgetClipping::ClipToBounds;
+	}
+	SLATE_END_ARGS()
+
+	void Construct(const FArguments& InArgs, TSharedPtr<SNPWindow> InNetworkPredictionWindow);
+
+	virtual void Tick(const FGeometry& AllottedGeometry, const double InCurrentTime, const float InDeltaTime) override;
+
+	virtual int32 OnPaint(const FPaintArgs& Args, const FGeometry& AllottedGeometry, const FSlateRect& MyCullingRect, FSlateWindowElementList& OutDrawElements, int32 LayerId, const FWidgetStyle& InWidgetStyle, bool bParentEnabled) const override;
+	virtual FReply OnMouseMove(const FGeometry& MyGeometry, const FPointerEvent& MouseEvent) override;
+	virtual FReply OnMouseButtonDown(const FGeometry& MyGeometry, const FPointerEvent& MouseEvent) override;
+	virtual FReply OnMouseButtonUp(const FGeometry& MyGeometry, const FPointerEvent& MouseEvent) override;
+	virtual void OnMouseEnter(const FGeometry& MyGeometry, const FPointerEvent& MouseEvent) override;
+	virtual void OnMouseLeave(const FPointerEvent& MouseEvent) override;
+	virtual FReply OnMouseWheel(const FGeometry& MyGeometry, const FPointerEvent& MouseEvent) override;
+
+	virtual void DrawHorizontalAxisGrid(FDrawContext& DrawContext, const FSlateBrush* Brush, const FSlateFontInfo& Font) const;
+
+	const FSimFrameViewport& GetViewport() const { return Viewport; }
+	const FVector2D& GetMousePosition() const { return MousePosition; }
+
+	bool GetAutoScroll() const { return bAutoScroll; }
+	void SetAutoScroll(bool bIn);
+	void SetAutoScrollDirty();
+
+	static float PulsePCT;
+
+	void OnGetOptionsMenu(FMenuBuilder& Builder);
+
+	void SearchUserData(const FText& InFilterText);
+
+private:
+
+	void Reset();
+	void UpdateState();
+	void BuildSimulationView_ActorGroups(const FFilteredDataCollection& FilteredDataCollection);
+	void BuildSimulationView_Tracks();
+	
+	// 
+	void UpdateDraw();
+
+	void HorizontalScrollBar_OnUserScrolled(float ScrollOffset);
+	void VerticalScrollBar_OnUserScrolled(float ScrollOffset);
+	void UpdateHorizontalScrollBar();
+	void UpdateVerticalScrollBar();
+	void OnUserScroll();
+
+	void ZoomHorizontally(const float Delta, const float X);	
+
+	struct FMouseOverInfo
+	{
+		const FSimulationActorGroup* Group = nullptr;
+		const FSimulationTrack* Track = nullptr;
+		const FSimulationTrack::FSubTrack* SubTrack = nullptr;
+		const FSimulationTrack::FSubTrack::FDrawNetRecv* DrawNetRecv = nullptr;
+		const FSimulationTrack::FSubTrack::FDrawTick* DrawSimTick = nullptr;
+
+		const FSimulationData::FNetSerializeRecv* GetNetRecv(uint64 MaxEngineFrame)
+		{
+			const FSimulationData::FNetSerializeRecv* Recv = DrawNetRecv ? &DrawNetRecv->Source.NetRecv : nullptr;
+			if (!Recv && DrawSimTick)
+			{
+				if (DrawSimTick->Source.Tick.StartNetRecv && DrawSimTick->Source.Tick.StartNetRecv->EngineFrame <= MaxEngineFrame)
+				{
+					Recv = DrawSimTick->Source.Tick.StartNetRecv;
+				}
+
+			}
+			return Recv;
+		}
+	};
+
+	FMouseOverInfo GetMouseOverInfo(FVector2D InMousePosition) const;
+
+
+	TSharedPtr<SNPWindow> NetworkPredictionWindow;
+
+	/** The track's viewport. Encapsulates info about position and scale. */
+	FSimFrameViewport Viewport;
+	bool bIsViewportDirty = false;
+
+	/** The slot that contains this widget's descendants.*/
+	TSharedPtr<SScrollBar> HorizontalScrollBar;
+	TSharedPtr<SScrollBar> VerticalScrollBar;
+
+	bool bIsStateDirty = false;	// What we want to draw has changed: underlying analyzed data or filtering/settings.
+	bool bIsDrawDirty = false;	// What we need to draw (may) have changed: scrolling, viewport changes etc can trigger this as well.
+	
+	bool bAutoScroll = true;			// User wants to auto scroll
+	bool bAutoScrollDirty = false;		// We need to update our view to scroll to latest
+
+	class IUnrealInsightsModule* InsightsModule;
+
+	//////////////////////////////////////////////////
+	// Panning and Zooming behaviors
+
+	/** The current mouse position. */
+	FVector2D MousePosition;
+
+	/** Mouse position during the call on mouse button down. */
+	FVector2D MousePositionOnButtonDown;
+	float ViewportPosXOnButtonDown;
+	float ViewportPosYOnButtonDown;
+
+	/** Mouse position during the call on mouse button up. */
+	FVector2D MousePositionOnButtonUp;
+
+	bool bIsLMB_Pressed = false;
+	bool bIsRMB_Pressed = false;
+
+	/** True, if the user is currently interactively scrolling the view (ex.: by holding the left mouse button and dragging). */
+	bool bIsScrolling = false;
+
+	//////////////////////////////////////////////////
+	// Misc
+
+	FGeometry ThisGeometry;
+
+	/** Cursor type. */
+	ECursorType CursorType;
+	
+	FSimulationFrameView SimulationFrameView;
+
+	FSimFrameHoverView HoverView;
+
+	mutable float TooltipOpacity = 0.f;
+	float PulseAccumulator=0;
+
+	bool CompactSimFrameView() const { return bCompactSimFrameView; }
+	bool LinearSimFrameView() const { return bLinearSimFrameView; }
+
+	void ToggleCompactSimFrameView() { bCompactSimFrameView = !bCompactSimFrameView; bIsStateDirty = true;  }
+	void ToggleLinearSimFrameview() { bLinearSimFrameView = !bLinearSimFrameView; bIsStateDirty = true; }
+
+	bool bCompactSimFrameView = false;
+	bool bLinearSimFrameView = true;
+
+	FSimTime ViewportMaxSimTimeMS = 0;
+	FSimTime ViewportMinSimTimeMS = 0;
+	float ViewportYMaxValue = 0.f;
+	float ViewportYContentStartValue = 30.f;
+
+	FString UserStateSearchString;
+	bool PerformSearch(const FSimulationData::FTick&, const FSimulationData::FRestrictedView& SimView);
+	bool PerformSearch(const FSimulationData::FNetSerializeRecv&, const FSimulationData::FRestrictedView& SimView);
+	bool PerformSearchInternal(const TArrayView<const FSimulationData::FUserState* const>& UserStates);
+
+	void SortActorGroupToTop(const FSimulationActorGroup* Group);
+
+	TArray<FSimNetActorID> UserSortedNetActors;
+};
\ No newline at end of file
diff --git a/Engine/Plugins/Runtime/NetworkPredictionInsights/Source/NetworkPredictionInsights/Private/UI/SNPToolbar.cpp b/Engine/Plugins/Runtime/NetworkPredictionInsights/Source/NetworkPredictionInsights/Private/UI/SNPToolbar.cpp
new file mode 100644
--- /dev/null
+++ b/Engine/Plugins/Runtime/NetworkPredictionInsights/Source/NetworkPredictionInsights/Private/UI/SNPToolbar.cpp
@@ -0,0 +1,242 @@
+// Copyright Epic Games, Inc. All Rights Reserved.
+
+#include "SNPToolbar.h"
+#include "Widgets/Input/SNumericEntryBox.h"
+#include "Widgets/Input/SSearchBox.h"
+
+#define LOCTEXT_NAMESPACE "NetworkPredictionInsightsModule"
+
+SNPToolbar::SNPToolbar()
+{
+}
+
+SNPToolbar::~SNPToolbar()
+{
+}
+
+void SNPToolbar::Construct(const FArguments& InArgs, TSharedPtr<SNPWindow> InParentWindow)
+{
+	ParentWindow = InParentWindow;
+
+	struct Local
+	{
+		static void FillViewToolbar(TSharedPtr<SNPWindow> ProfilerWindow, FToolBarBuilder& ToolbarBuilder)
+		{
+			ToolbarBuilder.BeginSection("View");
+			{
+				ToolbarBuilder.AddToolBarButton(FNetworkPredictionInsightsManager::GetCommands().ToggleAutoScrollSimulationFrames);
+			}
+			ToolbarBuilder.EndSection();
+			/*
+			ToolbarBuilder.BeginSection("Connection");
+			{
+				TSharedRef<SWidget> GameInstanceWidget = ProfilerWindow->CreateGameInstanceComboBox();
+				ToolbarBuilder.AddWidget(GameInstanceWidget);
+
+				TSharedRef<SWidget> ConnectionWidget = ProfilerWindow->CreateConnectionComboBox();
+				ToolbarBuilder.AddWidget(ConnectionWidget);
+
+				TSharedRef<SWidget> ConnectionModeWidget = ProfilerWindow->CreateConnectionModeComboBox();
+				ToolbarBuilder.AddWidget(ConnectionModeWidget);
+			}
+			ToolbarBuilder.EndSection();
+			ToolbarBuilder.BeginSection("New");
+			{
+				//TODO: ToolbarBuilder.AddToolBarButton(FNetworkPredictionManager::GetCommands().OpenNewWindow);
+			}
+			ToolbarBuilder.EndSection();
+			*/
+		}
+
+		static void FillRightSideToolbarPre(FToolBarBuilder& ToolbarBuilder)
+		{
+			ToolbarBuilder.BeginSection("Debug");
+			{
+				ToolbarBuilder.AddToolBarButton(FNetworkPredictionInsightsManager::GetCommands().FirstEngineFrame);
+				ToolbarBuilder.AddToolBarButton(FNetworkPredictionInsightsManager::GetCommands().PrevEngineFrame);
+			}
+			ToolbarBuilder.EndSection();
+		}
+		static void FillRightSideToolbarPost(FToolBarBuilder& ToolbarBuilder)
+		{
+			ToolbarBuilder.BeginSection("Debug");
+			{
+				ToolbarBuilder.AddToolBarButton(FNetworkPredictionInsightsManager::GetCommands().NextEngineFrame);
+				ToolbarBuilder.AddToolBarButton(FNetworkPredictionInsightsManager::GetCommands().LastEngineFrame);
+			}
+			ToolbarBuilder.EndSection();
+		}
+	};
+	
+
+	TSharedPtr<FUICommandList> CommandList = ParentWindow->GetCommandList();
+
+	FToolBarBuilder ToolbarBuilder(CommandList.ToSharedRef(), FMultiBoxCustomization::None);
+	Local::FillViewToolbar(ParentWindow, ToolbarBuilder);
+
+
+	FToolBarBuilder RightSideToolbarBuilderPrev(CommandList.ToSharedRef(), FMultiBoxCustomization::None);
+	Local::FillRightSideToolbarPre(RightSideToolbarBuilderPrev);
+
+	FToolBarBuilder RightSideToolbarBuilderPost(CommandList.ToSharedRef(), FMultiBoxCustomization::None);
+	Local::FillRightSideToolbarPost(RightSideToolbarBuilderPost);
+
+	// Create the tool bar
+	ChildSlot
+	[
+		SNew(SHorizontalBox)
+		
+#if 0
+		+SHorizontalBox::Slot()
+		.AutoWidth()
+		.VAlign(VAlign_Center)
+		.Padding(2.0f, 0.0f, 0.0f, 0.0f)
+		[
+			SNew(SComboButton)
+			.Visibility(EVisibility::Visible)
+			//.ComboButtonStyle(FSourceFilterStyle::Get(), "SourceFilter.ComboButton")
+			//.ForegroundColor(FLinearColor::White)
+			.ContentPadding(0.0f)
+			.OnGetMenuContent(this, &SNPToolbar::OnGetOptionsMenu)
+			.ButtonContent()
+			[
+				SNew(SHorizontalBox)
+				/*
+				+SHorizontalBox::Slot()
+				.AutoWidth()
+				.VAlign(VAlign_Center)
+				[
+					SNew(STextBlock)
+					//.TextStyle(FSourceFilterStyle::Get(), "SourceFilter.TextStyle")
+					//.Font(FEditorStyle::Get().GetFontStyle("FontAwesome.9"))
+					.Text(FText::FromString(FString(TEXT("\xf0b0"))))
+				]*/
+				+SHorizontalBox::Slot()
+				.AutoWidth()
+				.Padding(2, 0, 0, 0)
+				.VAlign(VAlign_Center)
+				[
+					SNew(STextBlock)
+					//.TextStyle(FSourceFilterStyle::Get(), "SourceFilter.TextStyle")
+					.Text(LOCTEXT("OptionMenuLabel", "Options"))
+				]
+			]
+		]
+#endif
+
+		+ SHorizontalBox::Slot()
+		.HAlign(HAlign_Left)
+		.VAlign(VAlign_Center)			
+		.Padding(2.0f)
+		.FillWidth(1.f)
+		[
+			SNew(SHorizontalBox)
+			
+			+ SHorizontalBox::Slot()
+			.VAlign(VAlign_Center)
+			.FillWidth(1.f)
+			.Padding(2.0f)
+			.MaxWidth(300.f)
+			[
+				SNew(SSearchBox)
+				.HintText(LOCTEXT("UserStateSearchBoxHint", "Search simulation states"))
+				.OnTextChanged(ParentWindow.ToSharedRef(), &SNPWindow::SearchUserData)
+				.ToolTipText(LOCTEXT("UserStateSearchHint", "Type here to search for content in user simulation states"))
+				.MinDesiredWidth(300.f)
+			]
+		
+
+			+SHorizontalBox::Slot()
+			.HAlign(HAlign_Left)
+			.VAlign(VAlign_Center)
+			.AutoWidth()
+			.Padding(0.0f)
+			[
+				SNew(SBorder)
+				.Padding(0)
+				.BorderImage(FEditorStyle::GetBrush("NoBorder"))
+				.IsEnabled(FSlateApplication::Get().GetNormalExecutionAttribute())
+				[
+					ToolbarBuilder.MakeWidget()
+				]
+			]
+
+			+SHorizontalBox::Slot()
+			.HAlign(HAlign_Left)
+			.VAlign(VAlign_Center)
+			.AutoWidth()
+			.Padding(0.0f)
+			[
+				SAssignNew(PIESessionComboBox, SComboBox<TSharedPtr<int32>>)
+				.ToolTipText(LOCTEXT("PIESessionComboBoxToolTip", "Filters view to a given session of PIE."))
+				.OptionsSource(&ParentWindow->PIESessionOptions)
+				.Visibility(MakeAttributeLambda([this] { return ParentWindow->PIESessionOptions.Num() > 1 ? EVisibility::Visible : EVisibility::Collapsed; }))
+				.OnSelectionChanged(ParentWindow.ToSharedRef(), &SNPWindow::PIESessionComboBox_OnSelectionChanged)
+				.OnGenerateWidget(ParentWindow.ToSharedRef(), &SNPWindow::PIESessionComboBox_OnGenerateWidget)
+				[
+					SNew(STextBlock)
+					.Text(ParentWindow.ToSharedRef(), &SNPWindow::PIESessionComboBox_GetSelectionText)
+				]
+			]
+		]
+			
+
+		+SHorizontalBox::Slot()
+		.HAlign(HAlign_Right)
+		.VAlign(VAlign_Center)
+		.AutoWidth()
+		[
+			
+			SNew(SHorizontalBox)		
+			+SHorizontalBox::Slot()
+			.HAlign(HAlign_Right)
+			.VAlign(VAlign_Center)
+			.AutoWidth()
+			.Padding(0.0f)
+			[
+				RightSideToolbarBuilderPrev.MakeWidget()
+			]
+		
+			+SHorizontalBox::Slot()
+			.HAlign(HAlign_Right)
+			.VAlign(VAlign_Center)
+			.AutoWidth()
+			.Padding(0.0f)
+			[
+				SNew(SBorder)
+				.Padding(0)
+				.BorderImage(FEditorStyle::GetBrush("NoBorder"))
+				.IsEnabled(FSlateApplication::Get().GetNormalExecutionAttribute())
+				[
+					SNew(SNumericEntryBox<uint64>)
+					.ToolTipText(LOCTEXT("EngineFrameToolTip", "Current Engine Frame the simulation timeline is viewing."))
+					.OnValueCommitted_Lambda([this](uint64 InValue, ETextCommit::Type InCommitType){ ParentWindow->SetEngineFrame(InValue); })
+					.Value_Lambda([this]() { return ParentWindow->GetCurrentEngineFrame(); } )
+					.MinValue_Lambda([this]() { return ParentWindow->GetMinEngineFrame(); } )
+					.MaxValue_Lambda([this]() { return ParentWindow->GetMaxEngineFrame(); } )
+				]
+			]
+
+			+SHorizontalBox::Slot()
+			.HAlign(HAlign_Right)
+			.VAlign(VAlign_Center)
+			.AutoWidth()
+			.Padding(0.0f)
+			[
+				RightSideToolbarBuilderPost.MakeWidget()
+			]
+		]
+	];
+}
+
+TSharedRef<SWidget> SNPToolbar::OnGetOptionsMenu()
+{
+	FMenuBuilder Builder(true, TSharedPtr<FUICommandList>());
+	if (ParentWindow.IsValid())
+	{
+		ParentWindow->OnGetOptionsMenu(Builder);
+	}
+	return Builder.MakeWidget();
+}
+
+#undef LOCTEXT_NAMESPACE
\ No newline at end of file
diff --git a/Engine/Plugins/Runtime/NetworkPredictionInsights/Source/NetworkPredictionInsights/Private/UI/SNPToolbar.h b/Engine/Plugins/Runtime/NetworkPredictionInsights/Source/NetworkPredictionInsights/Private/UI/SNPToolbar.h
new file mode 100644
--- /dev/null
+++ b/Engine/Plugins/Runtime/NetworkPredictionInsights/Source/NetworkPredictionInsights/Private/UI/SNPToolbar.h
@@ -0,0 +1,36 @@
+// Copyright Epic Games, Inc. All Rights Reserved.
+
+#pragma once
+
+#include "CoreMinimal.h"
+#include "Widgets/SCompoundWidget.h"
+
+///////////////////////////////////////////////////////////////////////////////////////////////////
+
+class SNPToolbar : public SCompoundWidget
+{
+public:
+	/** Default constructor. */
+	SNPToolbar();
+
+	/** Virtual destructor. */
+	virtual ~SNPToolbar();
+
+	SLATE_BEGIN_ARGS(SNPToolbar) {}
+	SLATE_END_ARGS()
+
+	/**
+	 * Construct this widget
+	 *
+	 * @param	InArgs	The declaration data for this widget
+	 */
+	void Construct(const FArguments& InArgs, TSharedPtr<SNPWindow> InProfilerWindow);
+
+private:
+	TSharedPtr<SNPWindow> ParentWindow;
+
+	TSharedRef<SWidget> OnGetOptionsMenu();
+
+	TSharedPtr<SComboBox<TSharedPtr<int32>>> PIESessionComboBox;
+
+};
\ No newline at end of file
diff --git a/Engine/Plugins/Runtime/NetworkPredictionInsights/Source/NetworkPredictionInsights/Private/UI/SNPWindow.cpp b/Engine/Plugins/Runtime/NetworkPredictionInsights/Source/NetworkPredictionInsights/Private/UI/SNPWindow.cpp
new file mode 100644
--- /dev/null
+++ b/Engine/Plugins/Runtime/NetworkPredictionInsights/Source/NetworkPredictionInsights/Private/UI/SNPWindow.cpp
@@ -0,0 +1,665 @@
+// Copyright Epic Games, Inc. All Rights Reserved.
+
+#include "SNPWindow.h"
+
+#include "EditorStyleSet.h"
+#include "Framework/Commands/UICommandList.h"
+#include "Framework/Docking/WorkspaceItem.h"
+#include "Framework/MultiBox/MultiBoxBuilder.h"
+#include "SlateOptMacros.h"
+#include "Widgets/Docking/SDockTab.h"
+#include "Widgets/Images/SImage.h"
+#include "Widgets/Layout/SBorder.h"
+#include "Widgets/Layout/SBox.h"
+#include "Widgets/Layout/SSpacer.h"
+#include "Widgets/Notifications/SNotificationList.h"
+#include "Widgets/SBoxPanel.h"
+#include "Widgets/Text/STextBlock.h"
+
+#include "Insights/Common/PaintUtils.h"
+#include "Widgets/Input/SNumericEntryBox.h"
+#include "SNPSimFrameContents.h"
+
+#define LOCTEXT_NAMESPACE "SNetworkPredictionWindow"
+
+
+SNPWindow::SNPWindow()
+{
+
+}
+
+SNPWindow::~SNPWindow()
+{
+
+}
+
+void SNPWindow::Construct(const FArguments& InArgs, const TSharedRef<SDockTab>& ConstructUnderMajorTab, const TSharedPtr<SWindow>& ConstructUnderWindow)
+{
+	InsightsModule = &FModuleManager::LoadModuleChecked<IUnrealInsightsModule>("TraceInsights");
+	CommandList = MakeShareable(new FUICommandList);
+
+	// Create & initialize tab manager.
+	TabManager = FGlobalTabmanager::Get()->NewTabManager(ConstructUnderMajorTab);
+	TSharedRef<FWorkspaceItem> AppMenuGroup = TabManager->AddLocalWorkspaceMenuCategory(LOCTEXT("NetworkPredictionMenuGroupName", "Network Prediction Insights"));
+
+	
+	TabManager->RegisterTabSpawner(FNetworkPredictionInsightsTabs::ToolbarID, FOnSpawnTab::CreateRaw(this, &SNPWindow::SpawnTab_Toolbar))
+		.SetDisplayName(LOCTEXT("DeviceToolbarTabTitle", "Toolbar"))
+		//.SetIcon(FSlateIcon(FInsightsStyle::GetStyleSetName(), "Toolbar.Icon.Small"))
+		.SetGroup(AppMenuGroup);
+	
+	TabManager->RegisterTabSpawner(FNetworkPredictionInsightsTabs::SimFrameViewID, FOnSpawnTab::CreateRaw(this, &SNPWindow::SpawnTab_SimFrameView))
+		.SetDisplayName(LOCTEXT("NetworkPredictionInsights.SimFrameViewTabTitle", "Simulation Frame View"))
+		//.SetIcon(FSlateIcon(FInsightsStyle::GetStyleSetName(), "PacketView.Icon.Small"))
+		.SetGroup(AppMenuGroup);
+
+	
+	TabManager->RegisterTabSpawner(FNetworkPredictionInsightsTabs::SimFrameContentsID, FOnSpawnTab::CreateRaw(this, &SNPWindow::SpawnTab_SimFrameContentView))
+		.SetDisplayName(LOCTEXT("NetworkPredictionInsights.SimFrameContentsTitle", "Simulation Frame Content"))
+		//.SetIcon(FSlateIcon(FInsightsStyle::GetStyleSetName(), "PacketContentView.Icon.Small"))
+		.SetGroup(AppMenuGroup);
+	/*
+
+	TabManager->RegisterTabSpawner(FNetworkingProfilerTabs::NetStatsViewID, FOnSpawnTab::CreateRaw(this, &SNetworkingProfilerWindow::SpawnTab_NetStatsView))
+		.SetDisplayName(LOCTEXT("NetworkingProfiler.NetStatsViewTabTitle", "Net Stats"))
+		.SetIcon(FSlateIcon(FInsightsStyle::GetStyleSetName(), "NetStatsView.Icon.Small"))
+		.SetGroup(AppMenuGroup);
+	*/
+	
+	//TSharedPtr<FNetworkPredictionManager> NetworkingProfilerManager = FNetworkPredictionManager::Get();
+	//ensure(NetworkingProfilerManager.IsValid());
+
+	// Create tab layout.
+	const TSharedRef<FTabManager::FLayout> Layout = FTabManager::NewLayout("NetworkingPredictionInsightsProfilerLayout_v1.0")
+		->AddArea
+		(
+			FTabManager::NewPrimaryArea()
+			->SetOrientation(Orient_Vertical)
+			
+			->Split
+			(
+				FTabManager::NewStack()
+				->AddTab(FNetworkPredictionInsightsTabs::ToolbarID, ETabState::OpenedTab)
+				->SetHideTabWell(true)
+				->SetSizeCoefficient(0.10f)
+			)
+			
+			->Split
+			(
+				FTabManager::NewSplitter()
+				->SetOrientation(Orient_Horizontal)
+				->SetSizeCoefficient(1.0f)
+				->Split
+				(
+					FTabManager::NewSplitter()
+					->SetOrientation(Orient_Vertical)
+					->SetSizeCoefficient(0.65f)
+					->Split
+					(
+						FTabManager::NewStack()
+						->SetSizeCoefficient(0.65f)
+						->SetHideTabWell(true)
+						->AddTab(FNetworkPredictionInsightsTabs::SimFrameViewID, ETabState::OpenedTab)
+					)
+					
+					->Split
+					(
+						FTabManager::NewStack()
+						->SetSizeCoefficient(0.35f)
+						->SetHideTabWell(true)
+						->AddTab(FNetworkPredictionInsightsTabs::SimFrameContentsID, ETabState::OpenedTab)
+					)
+				)
+				/*
+				->Split
+				(
+					FTabManager::NewStack()
+					->SetSizeCoefficient(0.35f)
+					//->AddTab(FNetworkingProfilerTabs::NetStatsViewID, ETabState::OpenedTab)
+					//->SetForegroundTab(FNetworkingProfilerTabs::NetStatsViewID)
+				)
+				*/
+			)
+		);
+
+	// Create & initialize main menu.
+	FMenuBarBuilder MenuBarBuilder = FMenuBarBuilder(CommandList);
+
+	MenuBarBuilder.AddPullDownMenu(
+		LOCTEXT("MenuLabel", "MENU"),
+		FText::GetEmpty(),
+		FNewMenuDelegate::CreateSP(this, &SNPWindow::FillMenu, TabManager),
+		FName(TEXT("MENU"))
+	);
+
+	ChildSlot
+		[
+			SNew(SOverlay)
+
+			// Version
+			+ SOverlay::Slot()
+				.HAlign(HAlign_Right)
+				.VAlign(VAlign_Top)
+				.Padding(0.0f, -16.0f, 0.0f, 0.0f)
+				[
+					SNew(STextBlock)
+						.Clipping(EWidgetClipping::ClipToBoundsWithoutIntersecting)
+						.Text(LOCTEXT("NetworkPredictionInsightsVersion", "0.1e-32"))
+						.ColorAndOpacity(FLinearColor(0.15f, 0.15f, 0.15f, 1.0f))
+				]
+
+			// Overlay slot for the main window area
+			+ SOverlay::Slot()
+				.HAlign(HAlign_Fill)
+				.VAlign(VAlign_Fill)
+				[
+					SNew(SVerticalBox)
+
+					+ SVerticalBox::Slot()
+						.AutoHeight()
+						[
+							MenuBarBuilder.MakeWidget()
+						]
+
+					+ SVerticalBox::Slot()
+						.FillHeight(1.0f)
+						[
+							TabManager->RestoreFrom(Layout, ConstructUnderWindow).ToSharedRef()
+						]
+				]
+
+			// Session hint overlay
+			/*
+			+ SOverlay::Slot()
+				.HAlign(HAlign_Center)
+				.VAlign(VAlign_Center)
+				[
+					SNew(SBorder)
+						//.Visibility(this, &SNetworkingProfilerWindow::IsSessionOverlayVisible)
+						.BorderImage(FEditorStyle::GetBrush("NotificationList.ItemBackground"))
+						.Padding(8.0f)
+						[
+							SNew(STextBlock)
+								.Text(LOCTEXT("SelectTraceOverlayText", "Please select a trace."))
+						]
+				]
+			*/
+		];
+
+	// Tell tab-manager about the global menu bar.
+	TabManager->SetMenuMultiBox(MenuBarBuilder.GetMultiBox());
+
+	BindCommands();
+
+	SetEngineFrame(0);
+}
+
+void SNPWindow::Tick(const FGeometry& AllottedGeometry, const double InCurrentTime, const float InDeltaTime)
+{
+	TSharedPtr<const Trace::IAnalysisSession> Session = InsightsModule->GetAnalysisSession();
+	if (Session.IsValid())
+	{
+		Trace::FAnalysisSessionReadScope SessionReadScope(*Session.Get());
+		if (const INetworkPredictionProvider* NetworkPredictionProvider = ReadNetworkPredictionProvider(*Session.Get()))
+		{
+			uint64 NewDataCounter = NetworkPredictionProvider->GetNetworkPredictionDataCounter();
+			if (NewDataCounter != CachedProviderDataCounter)
+			{
+				PopulateFilteredDataView();
+				CachedProviderDataCounter = NewDataCounter;
+			}
+		}
+	}
+}
+
+void SNPWindow::Reset()
+{
+
+}
+
+void SNPWindow::BindCommands()
+{
+	Map_ToggleAutoScrollSimulationFrames();
+
+	CommandList->MapAction(
+		FNetworkPredictionInsightsManager::GetCommands().NextEngineFrame,
+		FExecuteAction::CreateSP(this, &SNPWindow::NextEngineFrame),
+		FCanExecuteAction::CreateSP(this, &SNPWindow::CanNextEngineFrame),
+		EUIActionRepeatMode::RepeatEnabled);
+
+	CommandList->MapAction(
+		FNetworkPredictionInsightsManager::GetCommands().PrevEngineFrame,
+		FExecuteAction::CreateSP(this, &SNPWindow::PrevEngineFrame),
+		FCanExecuteAction::CreateSP(this, &SNPWindow::CanPrevEngineFrame),
+		EUIActionRepeatMode::RepeatEnabled);
+	
+	CommandList->MapAction(
+		FNetworkPredictionInsightsManager::GetCommands().FirstEngineFrame,
+		FExecuteAction::CreateSP(this, &SNPWindow::FirstEngineFrame),
+		FCanExecuteAction::CreateSP(this, &SNPWindow::CanFirstEngineFrame));
+
+	CommandList->MapAction(
+		FNetworkPredictionInsightsManager::GetCommands().LastEngineFrame,
+		FExecuteAction::CreateSP(this, &SNPWindow::LastEngineFrame),
+		FCanExecuteAction::CreateSP(this, &SNPWindow::CanLastEngineFrame));
+}
+
+FReply SNPWindow::OnKeyDown(const FGeometry& MyGeometry, const FKeyEvent& InKeyEvent)
+{
+	// This isn't doing anything and I cant figure out why
+	if(CommandList->ProcessCommandBindings(InKeyEvent.GetKey(), InKeyEvent.GetModifierKeys(), InKeyEvent.IsRepeat()))
+	{
+		return FReply::Handled();
+	}
+	return SCompoundWidget::OnKeyDown(MyGeometry, InKeyEvent);
+}
+
+const bool SNPWindow::AutoScrollSimulationFrames() const
+{
+	if (FrameView.IsValid())
+	{
+		return FrameView->GetAutoScroll();
+	}
+
+	return false;
+}
+void SNPWindow::SetAutoScrollSimulationFrames(const bool bAutoScroll)
+{
+	if (FrameView.IsValid())
+	{
+		FrameView->SetAutoScroll(bAutoScroll);
+	}
+}
+
+void SNPWindow::SetAutoScrollDirty()
+{
+	if (FrameView.IsValid())
+	{
+		FrameView->SetAutoScrollDirty();
+	}
+}
+
+////////////////////////////////////////////////////////////////////////////////////////////////////
+// Toggle Commands
+////////////////////////////////////////////////////////////////////////////////////////////////////
+
+#define IMPLEMENT_TOGGLE_COMMAND(CmdName, IsEnabled, SetEnabled) \
+	\
+	void SNPWindow::Map_##CmdName()\
+	{\
+		CommandList->MapAction(FNetworkPredictionInsightsManager::GetCommands().CmdName, CmdName##_Custom());\
+	}\
+	\
+	const FUIAction SNPWindow::CmdName##_Custom() \
+	{\
+		FUIAction UIAction;\
+		UIAction.ExecuteAction = FExecuteAction::CreateRaw(this, &SNPWindow::CmdName##_Execute);\
+		UIAction.CanExecuteAction = FCanExecuteAction::CreateRaw(this, &SNPWindow::CmdName##_CanExecute);\
+		UIAction.GetActionCheckState = FGetActionCheckState::CreateRaw(this, &SNPWindow::CmdName##_GetCheckState);\
+		return UIAction;\
+	}\
+	\
+	void SNPWindow::CmdName##_Execute()\
+	{\
+		SetEnabled(!IsEnabled());\
+	}\
+	\
+	bool SNPWindow::CmdName##_CanExecute() const\
+	{\
+		return true;\
+	}\
+	\
+	ECheckBoxState SNPWindow::CmdName##_GetCheckState() const\
+	{\
+		return IsEnabled() ? ECheckBoxState::Checked : ECheckBoxState::Unchecked;\
+	}
+
+IMPLEMENT_TOGGLE_COMMAND(ToggleAutoScrollSimulationFrames, AutoScrollSimulationFrames, SetAutoScrollSimulationFrames)
+#undef IMPLEMENT_TOGGLE_COMMAND
+
+void SNPWindow::FillMenu(FMenuBuilder& MenuBuilder, const TSharedPtr<FTabManager> InTabManager)
+{
+	if (!InTabManager.IsValid())
+	{
+		return;
+	}
+
+#if !WITH_EDITOR
+	//TODO: FGlobalTabmanager::Get()->PopulateTabSpawnerMenu(MenuBuilder, WorkspaceMenu::GetMenuStructure().GetStructureRoot());
+#endif //!WITH_EDITOR
+
+	TabManager->PopulateLocalTabSpawnerMenu(MenuBuilder);
+
+	OnGetOptionsMenu(MenuBuilder);
+}
+
+TSharedRef<SDockTab> SNPWindow::SpawnTab_Toolbar(const FSpawnTabArgs& Args)
+{
+	const TSharedRef<SDockTab> DockTab = SNew(SDockTab)
+		.ShouldAutosize(true)
+		.TabRole(ETabRole::PanelTab)
+		[
+			SNew(SNPToolbar, SharedThis(this))
+		];
+
+	DockTab->SetOnTabClosed(SDockTab::FOnTabClosedCallback::CreateRaw(this, &SNPWindow::OnToolbarTabClosed));
+	return DockTab;
+}
+
+TSharedRef<SDockTab> SNPWindow::SpawnTab_SimFrameView(const FSpawnTabArgs& Args)
+{
+	const TSharedRef<SDockTab> DockTab = SNew(SDockTab)
+		.ShouldAutosize(false)
+		.TabRole(ETabRole::PanelTab)
+		[
+			SAssignNew(FrameView, SNPSimFrameView, SharedThis(this))
+		];
+
+	DockTab->SetOnTabClosed(SDockTab::FOnTabClosedCallback::CreateRaw(this, &SNPWindow::OnSimFrameViewTabClosed));
+	return DockTab;
+}
+
+TSharedRef<SDockTab> SNPWindow::SpawnTab_SimFrameContentView(const FSpawnTabArgs& Args)
+{
+		const TSharedRef<SDockTab> DockTab = SNew(SDockTab)
+		.ShouldAutosize(false)
+		.TabRole(ETabRole::PanelTab)
+		[
+			SAssignNew(SimFrameContents, SNPSimFrameContents, SharedThis(this))
+		];
+
+	DockTab->SetOnTabClosed(SDockTab::FOnTabClosedCallback::CreateRaw(this, &SNPWindow::OnSimFrameContentTabClosed));
+	return DockTab;
+
+}
+
+void SNPWindow::OnSimFrameViewTabClosed(TSharedRef<SDockTab> TabBeingClosed)
+{
+	FrameView = nullptr;
+}
+
+void SNPWindow::OnToolbarTabClosed(TSharedRef<SDockTab> TabBeingClosed)
+{
+	
+}
+
+void SNPWindow::OnSimFrameContentTabClosed(TSharedRef<SDockTab> TabBeingClosed)
+{
+	SimFrameContents = nullptr;
+}
+
+////////////////////////////////////////////////////////////////////////////////////////////////
+
+void SNPWindow::SetEngineFrame(uint64 InFrame, const bool bSetAutoScrollDirty)
+{
+	MaxEngineFrameHistory.SetNum(EngineFrameHistoryIndex+1);
+	MaxEngineFrameHistory.Add(InFrame);
+	EngineFrameHistoryIndex = MaxEngineFrameHistory.Num()-1;
+	
+	if (MaxEngineFrameHistory.Num() > 32)
+	{
+		MaxEngineFrameHistory.RemoveAt(0, 16, false);
+		EngineFrameHistoryIndex -= 16;
+		check(EngineFrameHistoryIndex >= 0);
+	}
+
+	Filter.MaxEngineFrame = InFrame;
+	PopulateFilteredDataView();
+	if (bSetAutoScrollDirty)
+	{
+		SetAutoScrollDirty();
+	}
+}
+uint64 SNPWindow::GetCurrentEngineFrame() const
+{
+	return Filter.MaxEngineFrame == 0 ? UnfilteredDataView.LastEngineFrame : Filter.MaxEngineFrame;
+}
+uint64 SNPWindow::GetMinEngineFrame() const
+{
+	return UnfilteredDataView.FirstEngineFrame;
+
+}
+uint64 SNPWindow::GetMaxEngineFrame() const
+{
+	return UnfilteredDataView.LastEngineFrame;
+}
+
+void SNPWindow::NextEngineFrame()
+{
+	SetEngineFrame(FMath::Min<uint64>(GetCurrentEngineFrame()+1, GetMaxEngineFrame()));
+}
+bool SNPWindow::CanNextEngineFrame() const
+{
+	return (GetCurrentEngineFrame() < GetMaxEngineFrame());
+}
+
+void SNPWindow::PrevEngineFrame()
+{
+	// Don't get confused: the << < DVR > >>  is about setting MaxEngineFrame to visualize up to. We haven't implemented "engine frame window" yet.
+	SetEngineFrame(FMath::Max<uint64>(GetCurrentEngineFrame()-1, GetMinEngineFrame()));
+}
+bool SNPWindow::CanPrevEngineFrame() const
+{
+	return (GetCurrentEngineFrame() > GetMinEngineFrame());
+}
+
+void SNPWindow::FirstEngineFrame()
+{
+	SetEngineFrame(GetMinEngineFrame());
+	SetAutoScrollDirty();
+}
+bool SNPWindow::CanFirstEngineFrame() const
+{
+	return (GetCurrentEngineFrame() > GetMinEngineFrame());
+}
+
+void SNPWindow::LastEngineFrame()
+{
+	SetEngineFrame(0); // Ok, we are cheating. Set it to 0 to mean "uncapped" so that we'll start taking in live data again. This UI isn't great.
+	SetAutoScrollSimulationFrames(true);
+}
+bool SNPWindow::CanLastEngineFrame() const
+{
+	return (GetCurrentEngineFrame() < GetMaxEngineFrame());
+}
+
+void SNPWindow::NotifySimContentClicked(const FSimContentsView& Content)
+{
+	SelectedContent = Content;
+	if (SimFrameContents.IsValid())
+	{
+		SimFrameContents->NotifyContentClicked(SelectedContent);
+	}
+}
+
+void SNPWindow::JumpPreviousViewedEngineFrame()
+{
+	EngineFrameHistoryIndex = FMath::Max(0, EngineFrameHistoryIndex - 1);
+	
+	Filter.MaxEngineFrame = MaxEngineFrameHistory[EngineFrameHistoryIndex];
+	PopulateFilteredDataView();
+	SetAutoScrollSimulationFrames(true);
+}
+
+void SNPWindow::JumpNextViewedEngineFrame()
+{
+	EngineFrameHistoryIndex = FMath::Min(MaxEngineFrameHistory.Num()-1, EngineFrameHistoryIndex + 1);
+	
+	Filter.MaxEngineFrame = MaxEngineFrameHistory[EngineFrameHistoryIndex];
+	PopulateFilteredDataView();
+	SetAutoScrollSimulationFrames(true);
+}
+
+void SNPWindow::SearchUserData(const FText& InFilterText)
+{
+	if (FrameView.IsValid())
+	{
+		FrameView->SearchUserData(InFilterText);
+	}
+}
+
+void SNPWindow::PopulateFilteredDataView()
+{
+	FilteredDataCollection.Simulations.Reset();
+	FilteredDataCollection.FirstEngineFrame = TNumericLimits<uint64>::Max();
+	FilteredDataCollection.LastEngineFrame = 0;
+	
+	UnfilteredDataView.FirstEngineFrame = TNumericLimits<uint64>::Max();
+	UnfilteredDataView.LastEngineFrame = 0;
+
+	TSharedPtr<const Trace::IAnalysisSession> Session = InsightsModule->GetAnalysisSession();
+	if (Session.IsValid())
+	{
+		Trace::FAnalysisSessionReadScope SessionReadScope(*Session.Get());
+		const INetworkPredictionProvider* NetworkPredictionProvider = ReadNetworkPredictionProvider(*Session.Get());
+		if (NetworkPredictionProvider)
+		{
+			TArrayView<const TSharedRef<FSimulationData>> UnfilteredSimulationList = NetworkPredictionProvider->ReadSimulationData();
+			if (UnfilteredSimulationList.Num() == 0)
+			{
+				return;
+			}
+
+			TArray<TSharedPtr<const FSimulationData>> FilteredSimulationList;
+			FilteredSimulationList.Reserve(UnfilteredSimulationList.Num());
+
+			// Build PIE Session list
+			for (auto& SharedRef : UnfilteredSimulationList)
+			{
+				const int32 PIESession = SharedRef->ConstData.ID.PIESession;
+				if (!PIESessionOptions.FindByPredicate([PIESession](const TSharedPtr<int32>& SharedPtr) { return *SharedPtr == PIESession; }))
+				{
+					PIESessionOptions.Add(MakeShareable(new int32(PIESession)));
+					SetEngineFrame(0, true);
+					SelectedPIESession = PIESession;
+					SetAutoScrollSimulationFrames(true);
+					
+				}
+			}
+			check(PIESessionOptions.Num() > 0);
+			const uint64 FilteredPieSession = SelectedPIESession == -1 ? *PIESessionOptions.Last() : SelectedPIESession;
+
+			// Filter
+			for (auto& SharedRef : UnfilteredSimulationList)
+			{
+				if (SharedRef->ConstData.ID.NetGUID == 0)
+				{
+					// No NetGUID assigned yet (this would only happen first few frames in a live session)
+					continue;
+				}
+
+				// TODO: actually filter by group name or actor class, etc
+
+				if (FilteredPieSession != SharedRef->ConstData.ID.PIESession)
+				{
+					continue;
+				}
+
+				FilteredSimulationList.Add(SharedRef);
+			}
+
+			// Process Filtered Simulation List into restricted views on the engine frames we want to pass on
+			for (auto& SharedPtr : FilteredSimulationList)
+			{
+				const FSimulationData* SimData = SharedPtr.Get();
+				if (!ensure(SimData) || SimData->EOFState.Num() <= 0) // EOF data is currently required. Eventually EOF data may go away and this will be refactored
+				{
+					continue;
+				}
+
+				uint64 SimMinEngineFrame = TNumericLimits<uint64>::Max();
+				uint64 SimMaxEngineFrame = 0;
+				/*
+				if (SimData->Ticks.Num() > 0)
+				{
+					SimMinEngineFrame = FMath::Min<uint64>(SimMinEngineFrame, SimData->Ticks[0].EngineFrame);
+					SimMaxEngineFrame = FMath::Max<uint64>(SimMinEngineFrame, SimData->Ticks[SimData->Ticks.Num()-1].EngineFrame);
+				}
+				*/
+				if (SimData->EOFState.Num() > 0)
+				{
+					SimMinEngineFrame = FMath::Min<uint64>(SimMinEngineFrame, SimData->EOFState[0].EngineFrame);
+					SimMaxEngineFrame = FMath::Max<uint64>(SimMinEngineFrame, SimData->EOFState[SimData->EOFState.Num()-1].EngineFrame);
+				}
+				/*
+				if (SimData->NetRecv.Num() > 0)
+				{
+					SimMinEngineFrame = FMath::Min<uint64>(SimMinEngineFrame, SimData->NetRecv[0].EngineFrame);
+					SimMaxEngineFrame = FMath::Max<uint64>(SimMinEngineFrame, SimData->NetRecv[SimData->NetRecv.Num()-1].EngineFrame);
+				}
+				*/
+
+				if (SimMaxEngineFrame < SimMinEngineFrame)
+				{
+					continue;
+				}
+
+				UnfilteredDataView.FirstEngineFrame = FMath::Min(UnfilteredDataView.FirstEngineFrame, SimMinEngineFrame);
+				UnfilteredDataView.LastEngineFrame = FMath::Max(UnfilteredDataView.LastEngineFrame, SimMaxEngineFrame);
+
+				if (SimMaxEngineFrame < Filter.MinEngineFrame)
+				{
+					continue;
+				}
+
+				if (Filter.MaxEngineFrame != 0 && SimMinEngineFrame > Filter.MaxEngineFrame)
+				{
+					continue;
+				}
+
+				// Find the valid ranges for the simulation data
+				const uint64 FirstFilteredFrame = FMath::Max(SimMinEngineFrame, Filter.MinEngineFrame);
+				const uint64 LastFilteredFrame = Filter.MaxEngineFrame == 0 ? SimMaxEngineFrame : FMath::Min(SimMaxEngineFrame, Filter.MaxEngineFrame);
+
+				FilteredDataCollection.FirstEngineFrame = FMath::Min(FilteredDataCollection.FirstEngineFrame, FirstFilteredFrame);
+				FilteredDataCollection.LastEngineFrame = FMath::Max(FilteredDataCollection.LastEngineFrame, LastFilteredFrame);
+
+				FilteredDataCollection.Simulations.Emplace(SimData->MakeRestrictedView(Filter.MinEngineFrame, FilteredDataCollection.LastEngineFrame));
+			}
+
+		}
+	}
+
+	OnFilteredDataCollectionChange.Broadcast(FilteredDataCollection);
+}
+
+void SNPWindow::OnGetOptionsMenu(FMenuBuilder& Builder)
+{
+	if (FrameView.IsValid())
+	{
+		FrameView->OnGetOptionsMenu(Builder);
+	}
+}
+
+FText SNPWindow::PIESessionComboBox_GetSelectionText() const
+{
+	uint64 ActualSession = SelectedPIESession;
+	if (SelectedPIESession == -1 && PIESessionOptions.Num() > 0)
+	{
+		ActualSession = *PIESessionOptions.Last();
+	}
+
+	return FText::FromString( FString::Printf(TEXT("PIE Session: %d"), ActualSession) );
+}
+
+void SNPWindow::PIESessionComboBox_OnSelectionChanged(TSharedPtr<int32> NewPIEMode, ESelectInfo::Type SelectInfo)
+{
+	if (NewPIEMode.IsValid())
+	{
+		SelectedPIESession = *NewPIEMode;
+		PopulateFilteredDataView();
+		SetEngineFrame(UnfilteredDataView.LastEngineFrame, true);
+	}
+}
+
+TSharedRef<SWidget> SNPWindow::PIESessionComboBox_OnGenerateWidget(TSharedPtr<int32> InPIEMode) const
+{
+	return SNew(STextBlock)
+		.Text(FText::FromString(LexToString(*InPIEMode)));
+}
+
+// ------------------------------------------------------------------
+
+////////////////////////////////////////////////////////////////////////////////////////////////
+
+#undef LOCTEXT_NAMESPACE
\ No newline at end of file
diff --git a/Engine/Plugins/Runtime/NetworkPredictionInsights/Source/NetworkPredictionInsights/Private/UI/SNPWindow.h b/Engine/Plugins/Runtime/NetworkPredictionInsights/Source/NetworkPredictionInsights/Private/UI/SNPWindow.h
new file mode 100644
--- /dev/null
+++ b/Engine/Plugins/Runtime/NetworkPredictionInsights/Source/NetworkPredictionInsights/Private/UI/SNPWindow.h
@@ -0,0 +1,226 @@
+// Copyright Epic Games, Inc. All Rights Reserved.
+
+#pragma once
+
+#include "CoreMinimal.h"
+#include "Framework/Docking/TabManager.h"
+#include "Input/Reply.h"
+#include "Layout/Visibility.h"
+#include "Misc/Guid.h"
+#include "SlateFwd.h"
+#include "TraceServices/Model/NetProfiler.h"
+#include "Widgets/DeclarativeSyntaxSupport.h"
+#include "Widgets/Docking/SDockTab.h"
+#include "Widgets/Input/SComboBox.h"
+#include "Widgets/SCompoundWidget.h"
+
+#include "Insights/Common/PaintUtils.h"
+#include "INetworkPredictionProvider.h"
+#include "Containers/ArrayView.h"
+
+namespace Trace
+{
+	class IAnalysisService;
+	class IAnalysisSession;
+}
+
+class SNPSimFrameView;
+class SNPSimFrameContents;
+class SNPWindow;
+
+struct FNetworkPredictionInsightsTabs
+{
+	// Tab identifiers
+	static const FName ToolbarID;
+	static const FName SimFrameViewID;
+	static const FName SimFrameContentsID;
+};
+
+// --------------------------------------------------------------------------------------------------------------------------------------------
+//	Notes on the forms that our data takes
+//
+//	[Runtime] Runtime data is traced via UE_NP_TRACE macros
+//	[TraceStream] data is now in some generic trace binary format (saved to file or sent over TCP connection)
+
+// <----- Standalone UnrealInsights or in Editor. No longer dependent on NetworkPrediction module! ---->
+
+//	[Analyzed] FNetworkPredictionAnalyzer processes data from trace stream and pushes it to provider 
+//	[Provided] FNetworkPredictionProvider stores "all the data" ready for someone to consume it
+
+//	<---- UI Widgets ------>
+
+//	[Filtered] We want to limit the range or type of data that the UI should display
+//		-Filter frame range
+//		-Filter type of sim (movement, parametric, etc)
+//		-Filter by network role (e.g, exclude simulated proxy)
+//
+//	[Sampled] We may need to merge some data, e.g., > 1 sim frame per pixel requires us to combine data to be rendered by UI (e.g, zooming out)	// TODO
+//	[Culled] We still need to cull out data that is off screen, based on viewport state. We don't want to do this in ::OnPaint since that is called every render frame.
+//
+//	<----- Pixels on your screen! ---->
+// --------------------------------------------------------------------------------------------------------------------------------------------
+
+struct FNPFilter
+{
+	uint64 MinEngineFrame=0; // Inclusive (0 = "start at frame 0")
+	uint64 MaxEngineFrame=0; // 0 = open ended. Otherwise, exclusive (1 = "only frame 0")
+};
+
+// Some cached values about our data that is not filtered
+struct FNPUnfilteredDataView
+{
+	uint64 FirstEngineFrame = 0; // First Engine frame that had NP data
+	uint64 LastEngineFrame = 0; // Last engine frame that had NP data
+};
+
+
+// -------------------------------------------------------------------------------------------------
+//
+// -------------------------------------------------------------------------------------------------
+
+struct FFilter
+{
+	uint64 MinEngineFrame=0; // Inclusive (0 = "start at frame 0")
+	uint64 MaxEngineFrame=0; // 0 = open ended. Otherwise, exclusive (1 = "only frame 0")
+};
+
+// This is the main data structure that is passed into the various widgets. Widgets should avoid reading directly from the provider API
+struct FFilteredDataCollection
+{
+	TArray<TSharedRef<FSimulationData::FRestrictedView>> Simulations;
+
+	uint64 FirstEngineFrame = 0;
+	uint64 LastEngineFrame = 0;
+};
+
+// View into a specific piece of simulation content. E.g, something you click on in the sim timeline view.
+struct FSimContentsView
+{
+	TSharedPtr<FSimulationData::FRestrictedView> SimView;
+
+	const FSimulationData::FTick* SimTick = nullptr;
+	const FSimulationData::FNetSerializeRecv* NetRecv = nullptr;
+};
+
+////////////////////////////////////////////////////////////////////////////////////////////////////
+
+// Network Prediction Window, this is the top level / starting point for the NP Insights UI
+class SNPWindow : public SCompoundWidget
+{
+
+public:
+
+	SNPWindow();
+	virtual ~SNPWindow();
+
+	SLATE_BEGIN_ARGS(SNPWindow) {}
+	SLATE_END_ARGS()
+
+	void Construct(const FArguments& InArgs, const TSharedRef<SDockTab>& ConstructUnderMajorTab, const TSharedPtr<SWindow>& ConstructUnderWindow);
+	virtual void Tick(const FGeometry& AllottedGeometry, const double InCurrentTime, const float InDeltaTime) override;
+
+	void Reset();
+
+	const TSharedPtr<FUICommandList> GetCommandList() const { return CommandList; }
+	virtual FReply OnKeyDown(const FGeometry& MyGeometry, const FKeyEvent& InKeyEvent) override;
+
+#define DECLARE_TOGGLE_COMMAND(CmdName)\
+public:\
+	void Map_##CmdName(); /**< Maps UI command info CmdName with the specified UI command list. */\
+	const FUIAction CmdName##_Custom(); /**< UI action for CmdName command. */\
+private:\
+	void CmdName##_Execute(); /**< Handles FExecuteAction for CmdName. */\
+	bool CmdName##_CanExecute() const; /**< Handles FCanExecuteAction for CmdName. */\
+	ECheckBoxState CmdName##_GetCheckState() const; /**< Handles FGetActionCheckState for CmdName. */
+
+	DECLARE_TOGGLE_COMMAND(ToggleAutoScrollSimulationFrames)
+#undef DECLARE_TOGGLE_COMMAND
+
+public:
+
+	const bool AutoScrollSimulationFrames() const;
+	void SetAutoScrollSimulationFrames(const bool bAutoScroll);
+	void SetAutoScrollDirty(); // Forces one auto scroll but does not change user-selectable state
+
+	void SetEngineFrame(uint64 InFrame, const bool bSetAutoScrollDirty=false);
+	uint64 GetCurrentEngineFrame() const;
+	uint64 GetMinEngineFrame() const;
+	uint64 GetMaxEngineFrame() const;
+
+	void NextEngineFrame();
+	bool CanNextEngineFrame() const;
+
+	void PrevEngineFrame();
+	bool CanPrevEngineFrame() const;
+
+	void FirstEngineFrame();
+	bool CanFirstEngineFrame() const;
+
+	void LastEngineFrame();
+	bool CanLastEngineFrame() const;
+
+	void JumpPreviousViewedEngineFrame();
+	void JumpNextViewedEngineFrame();
+
+	// Called when the filtered view of our data has changed.
+	// Note this won't be called just because new data (within our filter) has been added by the provider
+	DECLARE_MULTICAST_DELEGATE_OneParam(FOnFilteredDataCollectionChange, const FFilteredDataCollection&)
+	FOnFilteredDataCollectionChange OnFilteredDataCollectionChange;
+
+	const FFilteredDataCollection& GetFilteredDataCollection() const { return FilteredDataCollection; }
+	
+	void OnGetOptionsMenu(FMenuBuilder& Builder);
+
+	// Notifies from SimFrame Timeline View
+	void NotifySimContentClicked(const FSimContentsView& Content);
+
+	const FSimContentsView& GetSelectedContent() const { return SelectedContent; }
+
+	void SearchUserData(const FText& InFilterText);
+
+
+	FText PIESessionComboBox_GetSelectionText() const;
+	void PIESessionComboBox_OnSelectionChanged(TSharedPtr<int32> NewPIEMode, ESelectInfo::Type SelectInfo);
+	TSharedRef<SWidget> PIESessionComboBox_OnGenerateWidget(TSharedPtr<int32> InPIEMode) const;
+
+	TArray<TSharedPtr<int32>> PIESessionOptions;
+	int32 SelectedPIESession = -1;
+
+private:
+
+	void PopulateFilteredDataView();
+
+	FNPFilter Filter;
+	FNPUnfilteredDataView UnfilteredDataView;
+
+	FFilteredDataCollection FilteredDataCollection;
+
+	void FillMenu(FMenuBuilder& MenuBuilder, const TSharedPtr<FTabManager> TabManager);
+
+	void BindCommands();	
+
+	/** Commandlist used in the window (Maps commands to window specific actions) */
+	TSharedPtr<FUICommandList> CommandList;
+
+	/** Holds the tab manager that manages the front-end's tabs. */
+	TSharedPtr<FTabManager> TabManager;
+
+	TSharedRef<SDockTab> SpawnTab_Toolbar(const FSpawnTabArgs& Args);
+	TSharedRef<SDockTab> SpawnTab_SimFrameView(const FSpawnTabArgs& Args);
+	TSharedRef<SDockTab> SpawnTab_SimFrameContentView(const FSpawnTabArgs& Args);
+	void OnToolbarTabClosed(TSharedRef<SDockTab> TabBeingClosed);
+	void OnSimFrameViewTabClosed(TSharedRef<SDockTab> TabBeingClosed);
+	void OnSimFrameContentTabClosed(TSharedRef<SDockTab> TabBeingClosed);
+
+	TSharedPtr<SNPSimFrameView> FrameView;
+	TSharedPtr<SNPSimFrameContents> SimFrameContents;
+
+	class IUnrealInsightsModule* InsightsModule;
+
+	FSimContentsView SelectedContent;
+	
+	uint64 CachedProviderDataCounter = 0;
+	
+	TArray<uint64> MaxEngineFrameHistory;
+	int32 EngineFrameHistoryIndex = -1;
+};
diff --git a/Engine/Plugins/Runtime/NetworkPredictionInsights/Source/NetworkPredictionInsights/Public/INetworkPredictionInsightsModule.h b/Engine/Plugins/Runtime/NetworkPredictionInsights/Source/NetworkPredictionInsights/Public/INetworkPredictionInsightsModule.h
new file mode 100644
--- /dev/null
+++ b/Engine/Plugins/Runtime/NetworkPredictionInsights/Source/NetworkPredictionInsights/Public/INetworkPredictionInsightsModule.h
@@ -0,0 +1,10 @@
+// Copyright Epic Games, Inc. All Rights Reserved.
+
+#pragma once
+
+#include "CoreMinimal.h"
+
+namespace NetworkPredictionInsightsTabs
+{
+	extern NETWORKPREDICTIONINSIGHTS_API const FName DocumentTab;
+};
diff --git a/Engine/Plugins/Runtime/NetworkPredictionInsights/Source/NetworkPredictionInsights/Public/INetworkPredictionProvider.h b/Engine/Plugins/Runtime/NetworkPredictionInsights/Source/NetworkPredictionInsights/Public/INetworkPredictionProvider.h
new file mode 100644
--- /dev/null
+++ b/Engine/Plugins/Runtime/NetworkPredictionInsights/Source/NetworkPredictionInsights/Public/INetworkPredictionProvider.h
@@ -0,0 +1,605 @@
+// Copyright Epic Games, Inc. All Rights Reserved.
+
+#pragma once
+
+#include "CoreMinimal.h"
+#include "TraceServices/Model/AnalysisSession.h"
+#include "TraceServices/Containers/Timelines.h"
+#include "Common/PagedArray.h"
+#include "TraceServices/Containers/Allocators.h"
+
+// In the NetworkPrediction system, simulations can individually specify their types used for timing.
+// Since we can't really specialize the insights side per simulation, we are just using single typedefs here that
+// could be compiled to larger target if necessary.
+
+#ifndef NP_INSIGHTS_SIM_TIME
+#define NP_INSIGHTS_SIM_TIME int32
+#endif
+
+using FSimTime = NP_INSIGHTS_SIM_TIME;
+
+// Mirroring this is unavoidable in order to keep NetworkPredictionInsights from depending on EngineTypes
+enum class ENP_NetRole: uint8
+{
+	/** No role at all. */
+	None,
+	/** Locally simulated proxy of this actor. */
+	SimulatedProxy,
+	/** Locally autonomous proxy of this actor. */
+	AutonomousProxy,
+	/** Authoritative control over the actor. */
+	Authority,
+	MAX,
+};
+
+enum class ENP_UserState: uint8
+{
+	Input,
+	Sync,
+	Aux,
+	MAX
+};
+
+enum class ENP_UserStateSource: uint8
+{
+	Unknown			= 1 << 0,	// not set yet
+	ProduceInput	= 1 << 1,	// System called produce input to sample new input locally (should only apply to InputCmd state)
+	SynthInput		= 1 << 2,	// System created input synthetically: based on previous state or some guess.
+	SimTick			= 1 << 3,	// Generated during a simulation tick
+	NetRecv			= 1 << 4,	// Received from network but not committed to buffers
+	NetRecvCommit	= 1 << 5,	// Received from network and committed to the buffers
+	OOB				= 1 << 6	// Out of band modification (e.g, game code)
+};
+
+enum class ENetSerializeRecvStatus: uint8
+{
+	Unknown,	// The moment we trace a Net Recv, its status is unknown.
+	
+	Confirm,	//	Valid data that did not cause a correction.
+	Rollback,	//	Caused simulation to rollback. A correction.
+	Jump,		//	Caused simulation to jump forward, effectively reseting its state
+	Fault,		//	Unable to process the data. For example, received frame we no longer had in our local buffers
+	Stale,		//	Received multiple bunches in a frame, this was effectively replaced by a newer one and not processed.
+};
+
+const TCHAR* LexToString(ENP_NetRole Role);
+const TCHAR* LexToString(ENP_UserState State);
+const TCHAR* LexToString(ENP_UserStateSource Source);
+const TCHAR* LexToString(ENetSerializeRecvStatus Status);
+
+// How we identify actors across network connection and PIE sessions
+struct FSimNetActorID
+{
+	uint32 NetGUID;
+	int32 PIESession;
+
+	bool operator==(const FSimNetActorID& Other) const
+	{
+		return NetGUID == Other.NetGUID && PIESession == Other.PIESession;
+	}
+};
+
+
+// ----------------------------------------------------------------------------------------------
+//
+// ----------------------------------------------------------------------------------------------
+
+template<typename T>
+struct TRestrictedPageArrayView;
+
+template<typename T>
+struct TRestrictedPageViewIterator
+{
+	TRestrictedPageViewIterator(const TRestrictedPageArrayView<T>& InPageArrayView, uint64 StartItemIndex=0)
+		: PageArrayView(InPageArrayView), CurrentItemIdx(StartItemIndex)
+	{
+		It = PageArrayView.PageArray.GetIteratorFromItem(StartItemIndex);
+	}
+
+	void operator++() { CurrentItemIdx++; It++;}
+	void operator--() { CurrentItemIdx--; It--;}
+	operator bool()	{ return Valid(); }
+
+	const T& operator* () const { if (Valid()) return *It.GetCurrentItem(); else { check(false); return *It.GetCurrentItem(); } }
+	const T* operator->() const { return Valid() ? It.GetCurrentItem() : nullptr; }	
+
+private:
+
+	bool Valid() const { return CurrentItemIdx >= PageArrayView.StartItemIdx && CurrentItemIdx < PageArrayView.EndItemIdx; }
+
+	const TRestrictedPageArrayView<T>& PageArrayView;
+	mutable TPagedArrayIterator<T, typename TPagedArray<T>::PageType> It;
+	uint64 CurrentItemIdx; // Kind of unfortunate but we store this ourselves to avoid calculating it off the iter each time
+};
+
+template<typename T>
+struct TRestrictedPageArrayView
+{
+	TRestrictedPageArrayView(const TPagedArray<T>& InPageArray, uint64 InStart, uint64 InEnd)
+		: PageArray(InPageArray), StartItemIdx(InStart), EndItemIdx(InEnd) { }
+
+	TRestrictedPageViewIterator<T> GetIterator() const
+	{
+		return TRestrictedPageViewIterator<T>(*this, StartItemIdx);
+	}
+
+	TRestrictedPageViewIterator<T> GetIteratorFromEnd() const
+	{
+		return TRestrictedPageViewIterator<T>(*this, EndItemIdx > 0 ? EndItemIdx-1 : 0);
+	}
+
+	const T& GetLast() const
+	{
+		return PageArray[EndItemIdx-1];
+	}
+
+	const T& GetFirst() const
+	{
+		return PageArray[StartItemIdx];
+	}
+
+	const T& operator[](uint64 Index) const
+	{
+		check(Index >= StartItemIdx && Index < EndItemIdx);
+		return PageArray[Index];
+	}
+
+private:
+
+	template<typename>
+	friend struct TRestrictedPageViewIterator;
+
+	const TPagedArray<T>& PageArray;
+	uint64 StartItemIdx;
+	uint64 EndItemIdx;
+};
+
+/*
+template<typename T>
+uint64 GetEngineFrame(const T& Element)
+{
+	return Element.EngineFrame;
+}
+*/
+
+template<typename T>
+typename TEnableIf<!TIsPointer<T>::Value, uint64>::Type GetEngineFrame(const T& Element)
+{
+	return Element.EngineFrame;
+}
+
+template<typename T>
+typename TEnableIf<TIsPointer<T>::Value, uint64>::Type GetEngineFrame(const T& Element)
+{
+	return Element->EngineFrame;
+}
+
+
+// Searches through PageArray for index of first element >= MinFrameNumber
+template<typename T>
+uint64 FindMinIndexPagedArray(const TPagedArray<T>& PageArray, uint64 MinFrameNumber)
+{
+	// Flip through pages first before iterating individual items
+	for (uint64 PageIdx=0; PageIdx < PageArray.NumPages(); ++PageIdx)
+	{
+		auto* Page = const_cast<TPagedArray<T>&>(PageArray).GetPage(PageIdx);
+		check(Page->Count >0);
+
+		//if (Page->Items[Page->Count-1].EngineFrame >= MinFrameNumber)
+		if (GetEngineFrame(Page->Items[Page->Count-1]) >= MinFrameNumber)
+		{
+			for (uint64 ItemIdx=0; ItemIdx < Page->Count; ++ItemIdx)
+			{
+				if (GetEngineFrame(Page->Items[ItemIdx]) >= MinFrameNumber)
+				{
+					return (PageIdx * PageArray.GetPageSize()) + ItemIdx;
+				}
+			}
+		}
+	}
+
+	return 0;
+}
+
+// Searches through PageArray for index+1 of first element > MinFrameNumber. E.g, index to iterate up to (<) for valid EngineFrameNumber.
+// Fixme: it probably makes more sense to just give this a start index and pickup where FindMinIndex left off
+template<typename T>
+uint64 FindMaxIndexPagedArray(const TPagedArray<T>& PageArray, uint64 MaxFrameNumber)
+{
+	if (PageArray.Num() == 0)
+	{
+		return 0;
+	}
+	
+	if (MaxFrameNumber == 0)
+	{
+		return PageArray.Num();
+	}
+
+	// Flip through pages first before iterating individual items
+	for (uint64 PageIdx=PageArray.NumPages()-1; ; --PageIdx)
+	{
+		auto* Page = const_cast<TPagedArray<T>&>(PageArray).GetPage(PageIdx);
+		check(Page->Count >0);
+
+		if (GetEngineFrame(Page->Items[0]) <= MaxFrameNumber)
+		{
+			for (uint64 ItemIdx=Page->Count-1; ; --ItemIdx)
+			{
+				if (GetEngineFrame(Page->Items[ItemIdx]) <= MaxFrameNumber)
+				{
+					return (PageIdx * PageArray.GetPageSize()) + ItemIdx + 1;
+				}
+
+				if (ItemIdx == 0)
+					break;
+
+			}
+		}
+
+		if (PageIdx == 0)
+			break;
+	}
+	return 0;
+}
+
+template<typename T>
+int32 FindMinIndexTArray(const TArray<const T*>& Array, uint64 MinFrameNumber)
+{
+	for (int32 idx=0; idx < Array.Num(); ++idx)
+	{
+		if (Array[idx]->EngineFrame >= MinFrameNumber)
+		{
+			return idx;
+		}
+	}
+
+	return 0;
+}
+
+template<typename T>
+int32 FindMaxIndexTArray(const TArray<const T*>& Array, uint64 MaxFrameNumber, int32 StartIdx=0)
+{
+	if (MaxFrameNumber == 0)
+	{
+		return Array.Num();
+	}
+
+	for (int32 idx=StartIdx; idx < Array.Num(); ++idx)
+	{
+		if (Array[idx]->EngineFrame > MaxFrameNumber)
+		{
+			return FMath::Max(0, idx);
+		}
+	}
+	return Array.Num();
+}
+
+template<typename T>
+struct TSparseFrameData
+{
+	const TSharedRef<const T> Read(uint64 EngineFrame) const
+	{
+		for (int32 idx=SparseData.Num()-1; idx >=0; --idx)
+		{
+			if (SparseData[idx]->EngineFrame <= EngineFrame)
+			{
+				return SparseData[idx];
+			}
+		}
+
+		// Didn't write anything for this frame so we insert default data
+		const_cast<TSparseFrameData<T>*>(this)->SparseData.Emplace(MakeShareable(new T(EngineFrame)));
+		return SparseData.Last();
+	}
+
+	const TSharedRef<T>& Write(uint64 EngineFrame)
+	{
+		for (int32 idx=SparseData.Num()-1; idx >=0; --idx)
+		{
+			if (SparseData[idx]->EngineFrame == EngineFrame)
+			{
+				// Already exists
+				return SparseData[idx];
+			}
+
+			if (SparseData[idx]->EngineFrame < EngineFrame)
+			{
+				// Must insert new
+				break;
+			}
+		}
+
+		SparseData.Emplace(MakeShareable(new T(EngineFrame)));
+		return SparseData.Last();
+	}
+
+private:
+
+	TArray<TSharedRef<T>> SparseData;
+};
+
+// Holds all data we traced for a given simulation.
+struct FSimulationData
+{
+	FSimulationData(uint32 InSimulationId, Trace::ILinearAllocator& Allocator)
+		: SimulationId(InSimulationId)
+		, Ticks(Allocator, 1024)
+		, EOFState(Allocator, 1024)
+		, NetRecv(Allocator, 1024)
+		, UserData(Allocator)
+	{ }
+
+	struct FTick;
+
+	struct FNetSerializeRecv
+	{
+		// Traced Data:
+		uint64 EngineFrame;
+		FSimTime SimTimeMS;
+		int32 Frame;
+
+		// Analysis Data:
+		ENetSerializeRecvStatus Status = ENetSerializeRecvStatus::Unknown;
+		bool bOrphan = true; // Hasn't been attached to FTick::StartNetRecv
+
+		FSimulationData::FTick* NextTick = nullptr;
+	};
+
+	// Data for a single simulation tick
+	struct FTick
+	{
+		// Traced Data:
+		uint64 EngineFrame;
+
+		FSimTime StartMS;
+		FSimTime EndMS;
+
+		int32 OutputFrame;
+
+		// Analysis Data:
+
+		// Net Recv that was processed prior to this tick. Will often be null. Used to ensure the ticks are drawn on the right sub track (Y)
+		const FNetSerializeRecv* StartNetRecv = nullptr;
+
+		uint64 ConfirmedEngineFrame = 0;	// Engine frame this became confirmed (we serialized a time past this)
+		uint64 TrashedEngineFrame = 0;		// Engine frame this became trashed (we resimulated this time in a later frame)
+		bool bRepredict = false;			// This tick was a repredict: we had already simulated up to this point before
+	};
+
+	// End of (Engine) Frame state of the simulation
+	// FIXME: it would be much nicer if we had a single global emission of GFrameNumber and EngineFrameDeltaTime
+	// It may be better to eventually just trace the events that change this simulation specific state (NetSerialize, Tick, etc)
+	struct FEngineFrame
+	{
+		uint64 EngineFrame;
+		double EngineFrameDeltaTime;
+		double EngineCurrentTime;
+
+		FSimTime	TotalProcessedSimTimeMS;
+		FSimTime	TotalAllowedSimTimeMS;
+		
+		int32		LastSentKeyframe;
+		int32		LastReceivedKeyframe;
+	};
+
+	// Data that changes rarely over time about the simulation
+	struct FSparse
+	{
+		FSparse(uint64 InEngineFrame)
+			: EngineFrame(InEngineFrame)
+		{}
+
+		uint64 EngineFrame;	// engine frame this data became valid
+		ENP_NetRole NetRole = ENP_NetRole::None;
+	};
+	
+	// Data that never changes about the simulation
+	struct FConst
+	{
+		FSimNetActorID ID;
+		uint32 GameInstanceId;
+		FName GroupName;
+		FString DebugName;
+	};
+
+	// -----------------------------------------------------------
+
+	struct FUserState
+	{
+		uint64 EngineFrame;	// EngineFrame this was written
+		int32 SimFrame; // This states sim frame (what frame# it was inserted into the buffer at)
+		FUserState* Prev = nullptr; // Previously written user state at this SimFrame number. Prev->EngineFrame < this->EngineFrame.
+		const TCHAR* UserStr = nullptr; // stable pointer to traced string representation of user state
+		ENP_UserStateSource Source = ENP_UserStateSource::Unknown; // Analyzed source of this user state
+	};
+
+	struct FUserStateStore
+	{
+		FUserStateStore(Trace::ILinearAllocator& Allocator)
+			: UserStates(Allocator, 1024) { }
+		
+		FORCENOINLINE FUserState& Push(int32 Frame, uint64 EngineFrame)
+		{
+			// Make new state
+			FUserState& NewElement = UserStates.PushBack();
+			NewElement.EngineFrame = EngineFrame;
+			NewElement.SimFrame = Frame;
+
+			// Put it in FrameMap
+			FUserState*& Head = FrameMap.FindOrAdd(Frame);
+			if (Head)
+			{
+				ensure(Head->EngineFrame <= EngineFrame);
+				NewElement.Prev = Head;
+			}
+			Head = &NewElement;
+
+			// Mark Frame as valid
+			const int32 Delta = (Frame+1) - PopulatedFrames.Num();
+			if (Delta > 0)
+			{
+				PopulatedFrames.Add(false, Delta);
+			}
+
+			PopulatedFrames[Frame] = true;
+			return NewElement;
+		}
+
+		FORCENOINLINE const FUserState* Get(int32 InFrame, uint64 MaxEngineFrame, uint8 ExcludeSourceMask, bool bSparseAccess) const
+		{
+			auto FindMatchAtFrame = [&](const int32 Frame)
+			{
+				// return first match at this Frame
+				FUserState* State = FrameMap.FindRef(Frame);
+				while (State)
+				{
+					if (State->EngineFrame <= MaxEngineFrame && (((uint8)State->Source & ExcludeSourceMask) == 0))
+					{
+						break;
+					}
+
+					State = State->Prev;
+				};
+
+				return State;
+			};
+
+			if (bSparseAccess)
+			{
+				if (FUserState* FoundMatch = FindMatchAtFrame(InFrame))
+				{
+					return FoundMatch;
+				}
+
+				// Reverse iterating through set bits not as easy as hoped but still beats continuous hash lookups
+				for (TBitArray<>::FConstReverseIterator BitIt(PopulatedFrames); BitIt; ++BitIt)
+				{
+					if (BitIt.GetIndex() < InFrame && BitIt.GetValue())
+					{
+						if (FUserState* FoundMatch = FindMatchAtFrame(BitIt.GetIndex()))
+						{
+							return FoundMatch;
+						}
+					}
+				}
+
+				return nullptr;
+			}
+
+			return FindMatchAtFrame(InFrame);
+		}
+	
+	private:
+
+		// Notes on data storage here:
+		//	-We may not have an FUserState for all Simulation Frames (interpolation case will be very sparse)
+		//	-We may have an FUserState for each simulation frame, or even multiple (autoproxy/correction heavy cases)
+		//	-User states should be pushed in ascending EngineFrame order but NOT necessarily in accessing SimulationFram 
+		//	-We don't need to hyper optimize this but we should avoid linear searching 
+
+		TPagedArray<FUserState> UserStates;	// Holds the actual FUserState. Pointers to these elements are stable
+		TMap<int32, FUserState*> FrameMap;	// Acceleration map to find head FUserState for given simulation frame
+		TBitArray<> PopulatedFrames;		// Bit array for which simulation frames are populated. Accelerates finding "last valid" state in sparse cases.
+	};
+
+	struct FUserData
+	{
+		FUserData(Trace::ILinearAllocator& Allocator)
+			: Store{Allocator, Allocator, Allocator}
+		{ }
+
+		FUserStateStore Store[ENP_UserState::MAX];
+	};
+
+	struct FRestrictedUserStateView
+	{
+		FRestrictedUserStateView(const FUserData& InUserData, uint64 InMaxEngineFrame)
+			: UserData(InUserData), MaxEngineFrame(InMaxEngineFrame)
+		{ }
+
+		FORCENOINLINE const FUserState* Get(ENP_UserState Type, int32 SimFrame, uint64 InEngineFrame, uint8 ExcludeMask) const
+		{
+			const uint64 EngineFrame = FMath::Min(InEngineFrame, MaxEngineFrame);
+			const bool bSparseAccess = (Type == ENP_UserState::Aux);
+			return UserData.Store[(int32)Type].Get(SimFrame, EngineFrame, ExcludeMask, bSparseAccess);
+		}
+
+	private:
+
+		const FUserData& UserData;
+		uint64 MaxEngineFrame;
+	};
+
+	// -----------------------------------------------------------
+
+	struct FRestrictedView
+	{
+		const TRestrictedPageArrayView<FTick> Ticks;
+		const TRestrictedPageArrayView<FEngineFrame> EOFState;
+		const TRestrictedPageArrayView<FNetSerializeRecv> NetRecv;
+		const FRestrictedUserStateView UserData;
+		const TSharedRef<const FSparse> SparseData;
+		const FConst& ConstData;
+		const uint32 SimulationId;
+	};
+
+	TSharedRef<FRestrictedView> MakeRestrictedView(uint64 MinEngineFrame, uint64 MaxEngineFrame) const
+	{
+		ensure(MaxEngineFrame != 0); // by now, "0 = uncapped" should be out of the picture
+
+		return MakeShareable(new FRestrictedView  {
+			{ Ticks, FindMinIndexPagedArray(Ticks, MinEngineFrame), FindMaxIndexPagedArray(Ticks, MaxEngineFrame) },
+			{ EOFState, FindMinIndexPagedArray(EOFState, MinEngineFrame), FindMaxIndexPagedArray(EOFState, MaxEngineFrame) },
+			{ NetRecv, FindMinIndexPagedArray(NetRecv, MinEngineFrame), FindMaxIndexPagedArray(NetRecv, MaxEngineFrame) },
+			FRestrictedUserStateView(UserData, MaxEngineFrame),
+			SparseData.Read(MaxEngineFrame), ConstData, SimulationId
+		});
+	}
+
+	// -----------------------------------------------------------
+
+	uint32 SimulationId;
+	
+	TPagedArray<FTick> Ticks;
+	TPagedArray<FEngineFrame> EOFState;
+	TPagedArray<FNetSerializeRecv> NetRecv;	// Actually holds the allocated recv records
+	
+	TSparseFrameData<FSparse> SparseData;
+	FConst ConstData;
+
+	FUserData UserData;
+
+	// Data that is only used as a temp/scratch pad during the analysis page. This should not show up in the UI or FRestrictedView
+	struct FAnalysis
+	{
+		// NetRecvs that we haven't matched with a tick or declared orphaned yet
+		TArray<FNetSerializeRecv*> PendingNetSerializeRecv;
+
+		uint64 NetRecvItemIdx = 0;
+		uint64 TrashItemIdx = 0;
+
+		FSimTime MaxTickSimTimeMS = 0; // Highest EndMS we have seen so far from Tick traces
+
+		ENP_UserStateSource PendingUserStateSource = ENP_UserStateSource::Unknown;
+		TArray<FUserState*> PendingCommitUserStates; // NetRecv'd state that hasn't been commited
+
+	} Analysis;
+};
+
+class INetworkPredictionProvider : public Trace::IProvider
+{
+public:
+
+	// Return the version reported in the trace
+	// A return value of 0 indicates no NetworkPrediction trace data
+	virtual uint32 GetNetworkPredictionTraceVersion() const = 0;
+
+	// Returns generic counter that increments everytime we analyse network prediction related data.
+	virtual uint64 GetNetworkPredictionDataCounter() const = 0;	
+	
+	// Returns all simulation data that has been analyzed
+	virtual TArrayView<const TSharedRef<FSimulationData>> ReadSimulationData() const = 0;
+};
+
+NETWORKPREDICTIONINSIGHTS_API const INetworkPredictionProvider* ReadNetworkPredictionProvider(const Trace::IAnalysisSession& Session);
\ No newline at end of file
diff --git a/Engine/Programs/UnrealInsights/Config/DefaultEngine.ini b/Engine/Programs/UnrealInsights/Config/DefaultEngine.ini
--- a/Engine/Programs/UnrealInsights/Config/DefaultEngine.ini
+++ b/Engine/Programs/UnrealInsights/Config/DefaultEngine.ini
@@ -1,3 +1,4 @@
 [Plugins]
 +ProgramEnabledPlugins=GameplayInsights
-+ProgramEnabledPlugins=TraceDataFilters
\ No newline at end of file
++ProgramEnabledPlugins=TraceDataFilters
++ProgramEnabledPlugins=NetworkPredictionInsights
\ No newline at end of file
