speculative -strict compile error fixes

diff --git a/Engine/Plugins/Runtime/NetworkPredictionInsights/Source/NetworkPredictionInsights/Private/UI/NetworkPredictionInsightsManager.h b/Engine/Plugins/Runtime/NetworkPredictionInsights/Source/NetworkPredictionInsights/Private/UI/NetworkPredictionInsightsManager.h
--- a/Engine/Plugins/Runtime/NetworkPredictionInsights/Source/NetworkPredictionInsights/Private/UI/NetworkPredictionInsightsManager.h
+++ b/Engine/Plugins/Runtime/NetworkPredictionInsights/Source/NetworkPredictionInsights/Private/UI/NetworkPredictionInsightsManager.h
@@ -4,7 +4,8 @@
 
 #include "CoreMinimal.h"
 #include "NetworkPredictionInsightsActionManager.h"
-//#include "Insights/InsightsManager.h"
+
+class SNPWindow;
 
 /**
  * This class manages the Networking Profiler state and settings.
diff --git a/Engine/Plugins/Runtime/NetworkPredictionInsights/Source/NetworkPredictionInsights/Public/INetworkPredictionProvider.h b/Engine/Plugins/Runtime/NetworkPredictionInsights/Source/NetworkPredictionInsights/Public/INetworkPredictionProvider.h
--- a/Engine/Plugins/Runtime/NetworkPredictionInsights/Source/NetworkPredictionInsights/Public/INetworkPredictionProvider.h
+++ b/Engine/Plugins/Runtime/NetworkPredictionInsights/Source/NetworkPredictionInsights/Public/INetworkPredictionProvider.h
@@ -509,7 +509,7 @@ struct FSimulationData
 			: Store{Allocator, Allocator, Allocator}
 		{ }
 
-		FUserStateStore Store[ENP_UserState::MAX];
+		FUserStateStore Store[(int32)ENP_UserState::MAX];
 	};
 
 	struct FRestrictedUserStateView
