NetworkPrediction Insights non unity fixes

diff --git a/Engine/Plugins/Runtime/NetworkPredictionInsights/Source/NetworkPredictionInsights/Private/UI/NetworkPredictionInsightsCommands.h b/Engine/Plugins/Runtime/NetworkPredictionInsights/Source/NetworkPredictionInsights/Private/UI/NetworkPredictionInsightsCommands.h
--- a/Engine/Plugins/Runtime/NetworkPredictionInsights/Source/NetworkPredictionInsights/Private/UI/NetworkPredictionInsightsCommands.h
+++ b/Engine/Plugins/Runtime/NetworkPredictionInsights/Source/NetworkPredictionInsights/Private/UI/NetworkPredictionInsightsCommands.h
@@ -3,6 +3,7 @@
 #pragma once
 
 #include "CoreMinimal.h"
+#include "Framework/Commands/Commands.h"
 
 ////////////////////////////////////////////////////////////////////////////////////////////////////
 
diff --git a/Engine/Plugins/Runtime/NetworkPredictionInsights/Source/NetworkPredictionInsights/Private/UI/NetworkPredictionInsightsManager.cpp b/Engine/Plugins/Runtime/NetworkPredictionInsights/Source/NetworkPredictionInsights/Private/UI/NetworkPredictionInsightsManager.cpp
--- a/Engine/Plugins/Runtime/NetworkPredictionInsights/Source/NetworkPredictionInsights/Private/UI/NetworkPredictionInsightsManager.cpp
+++ b/Engine/Plugins/Runtime/NetworkPredictionInsights/Source/NetworkPredictionInsights/Private/UI/NetworkPredictionInsightsManager.cpp
@@ -1,57 +1,37 @@
 // Copyright Epic Games, Inc. All Rights Reserved.
 
 #include "NetworkPredictionInsightsManager.h"
-
-const FName FNetworkPredictionInsightsTabs::ToolbarID("Toolbar");
-const FName FNetworkPredictionInsightsTabs::SimFrameViewID("SimFrameView");
-const FName FNetworkPredictionInsightsTabs::SimFrameContentsID("SimFrameContents");
-
+#include "SNPWindow.h"
+#include "NetworkPredictionInsightsCommands.h"
 
 TSharedPtr<FNetworkPredictionInsightsManager> FNetworkPredictionInsightsManager::Instance = nullptr;
 
 ////////////////////////////////////////////////////////////////////////////////////////////////////
 
 FNetworkPredictionInsightsManager::FNetworkPredictionInsightsManager()
-	//: CommandList(MakeShared<FUICommandList>())
 	: ActionManager(this)
 {
-	//FInsightsManager::Get()->GetCommandList()
 }
 
 void FNetworkPredictionInsightsManager::PostConstructor()
 {
-	// Register tick functions.
-	//OnTick = FTickerDelegate::CreateSP(this, &FNetworkPredictionManager::Tick);
-	//OnTickHandle = FTicker::GetCoreTicker().AddTicker(OnTick, 1.0f);
-
 	FNetworkPredictionInsightsCommands::Register();
 	BindCommands();
 }
 
 void FNetworkPredictionInsightsManager::BindCommands()
 {
-	//ActionManager.Map_
-	//CommandList->MapAction(GetCommands().NextEngineFrame, FUIAction::FExecuteAction
 }
 
 FNetworkPredictionInsightsManager::~FNetworkPredictionInsightsManager()
 {
 	FNetworkPredictionInsightsCommands::Unregister();
-
-	// Unregister tick function.
-	//FTicker::GetCoreTicker().RemoveTicker(OnTickHandle);
 }
 
 TSharedPtr<FNetworkPredictionInsightsManager> FNetworkPredictionInsightsManager::Get()
 {
 	return FNetworkPredictionInsightsManager::Instance;
 }
-/*
-const TSharedRef<FUICommandList> FNetworkPredictionInsightsManager::GetCommandList() const
-{
-	return CommandList;
-}
-*/
 
 const FNetworkPredictionInsightsCommands& FNetworkPredictionInsightsManager::GetCommands()
 {
@@ -65,10 +45,24 @@ FNetworkPredictionInsightsActionManager& FNetworkPredictionInsightsManager::GetA
 
 bool FNetworkPredictionInsightsManager::Tick(float DeltaTime)
 {
-	//SCOPE_CYCLE_COUNTER(STAT_IOPM_Tick);
 	return true;
 }
 
+void FNetworkPredictionInsightsManager::AddProfilerWindow(const TSharedRef<SNPWindow>& InWindow)
+{
+	NetworkPredictionInsightsWindows.Add(InWindow);
+}
+
+void FNetworkPredictionInsightsManager::RemoveProfilerWindow(const TSharedRef<SNPWindow>& InWindow)
+{
+	NetworkPredictionInsightsWindows.Remove(InWindow);
+}
+	
+TSharedPtr<class SNPWindow> FNetworkPredictionInsightsManager::GetProfilerWindow(int32 Index) const
+{
+	return NetworkPredictionInsightsWindows[Index].Pin();
+}
+
 void FNetworkPredictionInsightsManager::OnSessionChanged()
 {
 	for (TWeakPtr<SNPWindow> WndWeakPtr : NetworkPredictionInsightsWindows)
diff --git a/Engine/Plugins/Runtime/NetworkPredictionInsights/Source/NetworkPredictionInsights/Private/UI/NetworkPredictionInsightsManager.h b/Engine/Plugins/Runtime/NetworkPredictionInsights/Source/NetworkPredictionInsights/Private/UI/NetworkPredictionInsightsManager.h
--- a/Engine/Plugins/Runtime/NetworkPredictionInsights/Source/NetworkPredictionInsights/Private/UI/NetworkPredictionInsightsManager.h
+++ b/Engine/Plugins/Runtime/NetworkPredictionInsights/Source/NetworkPredictionInsights/Private/UI/NetworkPredictionInsightsManager.h
@@ -4,6 +4,7 @@
 
 #include "CoreMinimal.h"
 #include "NetworkPredictionInsightsActionManager.h"
+#include "Containers/Ticker.h"
 
 class SNPWindow;
 
@@ -67,25 +68,16 @@ public:
 	/** @return an instance of the Networking Profiler action manager. */
 	static FNetworkPredictionInsightsActionManager& GetActionManager();
 
-	void AddProfilerWindow(const TSharedRef<SNPWindow>& InWindow)
-	{
-		NetworkPredictionInsightsWindows.Add(InWindow);
-	}
+	void AddProfilerWindow(const TSharedRef<SNPWindow>& InWindow);
 
-	void RemoveProfilerWindow(const TSharedRef<SNPWindow>& InWindow)
-	{
-		NetworkPredictionInsightsWindows.Remove(InWindow);
-	}
+	void RemoveProfilerWindow(const TSharedRef<SNPWindow>& InWindow);
 
 	/**
 	 * Converts profiler window weak pointer to a shared pointer and returns it.
 	 * Make sure the returned pointer is valid before trying to dereference it.
 	 */
-	TSharedPtr<class SNPWindow> GetProfilerWindow(int32 Index) const
-	{
-		return NetworkPredictionInsightsWindows[Index].Pin();
-	}
-
+	TSharedPtr<class SNPWindow> GetProfilerWindow(int32 Index) const;
+	
 	////////////////////////////////////////////////////////////////////////////////////////////////////
 
 	void OnSessionChanged();
diff --git a/Engine/Plugins/Runtime/NetworkPredictionInsights/Source/NetworkPredictionInsights/Private/UI/SNPWindow.cpp b/Engine/Plugins/Runtime/NetworkPredictionInsights/Source/NetworkPredictionInsights/Private/UI/SNPWindow.cpp
--- a/Engine/Plugins/Runtime/NetworkPredictionInsights/Source/NetworkPredictionInsights/Private/UI/SNPWindow.cpp
+++ b/Engine/Plugins/Runtime/NetworkPredictionInsights/Source/NetworkPredictionInsights/Private/UI/SNPWindow.cpp
@@ -19,9 +19,18 @@
 #include "Insights/Common/PaintUtils.h"
 #include "Widgets/Input/SNumericEntryBox.h"
 #include "SNPSimFrameContents.h"
+#include "Modules/ModuleManager.h"
+#include "Insights/IUnrealInsightsModule.h"
+#include "NetworkPredictionInsightsManager.h"
+#include "NetworkPredictionInsightsCommands.h"
+#include "SNPSimFrameView.h"
+#include "SNPToolbar.h"
 
 #define LOCTEXT_NAMESPACE "SNetworkPredictionWindow"
 
+const FName FNetworkPredictionInsightsTabs::ToolbarID("Toolbar");
+const FName FNetworkPredictionInsightsTabs::SimFrameViewID("SimFrameView");
+const FName FNetworkPredictionInsightsTabs::SimFrameContentsID("SimFrameContents");
 
 SNPWindow::SNPWindow()
 {
diff --git a/Engine/Plugins/Runtime/NetworkPredictionInsights/Source/NetworkPredictionInsights/Private/UI/SNPWindow.h b/Engine/Plugins/Runtime/NetworkPredictionInsights/Source/NetworkPredictionInsights/Private/UI/SNPWindow.h
--- a/Engine/Plugins/Runtime/NetworkPredictionInsights/Source/NetworkPredictionInsights/Private/UI/SNPWindow.h
+++ b/Engine/Plugins/Runtime/NetworkPredictionInsights/Source/NetworkPredictionInsights/Private/UI/SNPWindow.h
@@ -17,6 +17,7 @@
 #include "Insights/Common/PaintUtils.h"
 #include "INetworkPredictionProvider.h"
 #include "Containers/ArrayView.h"
+#include "Framework/Commands/UICommandList.h"
 
 namespace Trace
 {
