more non unity fixes

diff --git a/Engine/Plugins/Runtime/NetworkPredictionInsights/Source/NetworkPredictionInsights/Private/UI/SNPSimFrameContents.h b/Engine/Plugins/Runtime/NetworkPredictionInsights/Source/NetworkPredictionInsights/Private/UI/SNPSimFrameContents.h
--- a/Engine/Plugins/Runtime/NetworkPredictionInsights/Source/NetworkPredictionInsights/Private/UI/SNPSimFrameContents.h
+++ b/Engine/Plugins/Runtime/NetworkPredictionInsights/Source/NetworkPredictionInsights/Private/UI/SNPSimFrameContents.h
@@ -6,6 +6,7 @@
 #include "Layout/Geometry.h"
 #include "Styling/SlateTypes.h"
 #include "Widgets/SCompoundWidget.h"
+#include "SNPWindow.h"
 
 class SNPWindow;
 
diff --git a/Engine/Plugins/Runtime/NetworkPredictionInsights/Source/NetworkPredictionInsights/Private/UI/SNPSimFrameView.cpp b/Engine/Plugins/Runtime/NetworkPredictionInsights/Source/NetworkPredictionInsights/Private/UI/SNPSimFrameView.cpp
--- a/Engine/Plugins/Runtime/NetworkPredictionInsights/Source/NetworkPredictionInsights/Private/UI/SNPSimFrameView.cpp
+++ b/Engine/Plugins/Runtime/NetworkPredictionInsights/Source/NetworkPredictionInsights/Private/UI/SNPSimFrameView.cpp
@@ -2,6 +2,8 @@
 
 #include "SNPSimFrameView.h"
 #include "NetworkPredictionDrawHelpers.h"
+#include "SNPWindow.h"
+#include "Insights/IUnrealInsightsModule.h"
 
 #define LOCTEXT_NAMESPACE "NetworkPredictionInsights"
 
diff --git a/Engine/Plugins/Runtime/NetworkPredictionInsights/Source/NetworkPredictionInsights/Private/UI/SNPSimFrameView.h b/Engine/Plugins/Runtime/NetworkPredictionInsights/Source/NetworkPredictionInsights/Private/UI/SNPSimFrameView.h
--- a/Engine/Plugins/Runtime/NetworkPredictionInsights/Source/NetworkPredictionInsights/Private/UI/SNPSimFrameView.h
+++ b/Engine/Plugins/Runtime/NetworkPredictionInsights/Source/NetworkPredictionInsights/Private/UI/SNPSimFrameView.h
@@ -5,6 +5,7 @@
 #include "NPAxisViewportInt32.h"
 #include "NPAxisViewportDouble.h"
 #include "INetworkPredictionProvider.h"
+#include "Insights/Common/PaintUtils.h"
 
 enum class ESimFrameStatus : uint8
 {
@@ -149,6 +150,8 @@ struct FSimFrameHoverView
 	}
 };
 
+class SNPWindow;
+
 ////////////////////////////////////////////////////////////////////////////////////////////////////
 
 // Viewport that is simulation frame based (not real time!)
@@ -325,7 +328,7 @@ private:
 
 	void Reset();
 	void UpdateState();
-	void BuildSimulationView_ActorGroups(const FFilteredDataCollection& FilteredDataCollection);
+	void BuildSimulationView_ActorGroups(const struct FFilteredDataCollection& FilteredDataCollection);
 	void BuildSimulationView_Tracks();
 	
 	// 
diff --git a/Engine/Plugins/Runtime/NetworkPredictionInsights/Source/NetworkPredictionInsights/Private/UI/SNPToolbar.cpp b/Engine/Plugins/Runtime/NetworkPredictionInsights/Source/NetworkPredictionInsights/Private/UI/SNPToolbar.cpp
--- a/Engine/Plugins/Runtime/NetworkPredictionInsights/Source/NetworkPredictionInsights/Private/UI/SNPToolbar.cpp
+++ b/Engine/Plugins/Runtime/NetworkPredictionInsights/Source/NetworkPredictionInsights/Private/UI/SNPToolbar.cpp
@@ -3,6 +3,9 @@
 #include "SNPToolbar.h"
 #include "Widgets/Input/SNumericEntryBox.h"
 #include "Widgets/Input/SSearchBox.h"
+#include "SNPWindow.h"
+#include "NetworkPredictionInsightsManager.h"
+#include "NetworkPredictionInsightsCommands.h"
 
 #define LOCTEXT_NAMESPACE "NetworkPredictionInsightsModule"
 
diff --git a/Engine/Plugins/Runtime/NetworkPredictionInsights/Source/NetworkPredictionInsights/Private/UI/SNPToolbar.h b/Engine/Plugins/Runtime/NetworkPredictionInsights/Source/NetworkPredictionInsights/Private/UI/SNPToolbar.h
--- a/Engine/Plugins/Runtime/NetworkPredictionInsights/Source/NetworkPredictionInsights/Private/UI/SNPToolbar.h
+++ b/Engine/Plugins/Runtime/NetworkPredictionInsights/Source/NetworkPredictionInsights/Private/UI/SNPToolbar.h
@@ -6,6 +6,7 @@
 #include "Widgets/SCompoundWidget.h"
 
 ///////////////////////////////////////////////////////////////////////////////////////////////////
+class SNPWindow;
 
 class SNPToolbar : public SCompoundWidget
 {
