Missed includes, again, some how

diff --git a/Engine/Plugins/Runtime/NetworkPredictionInsights/Source/NetworkPredictionInsights/Private/UI/SNPSimFrameContents.cpp b/Engine/Plugins/Runtime/NetworkPredictionInsights/Source/NetworkPredictionInsights/Private/UI/SNPSimFrameContents.cpp
--- a/Engine/Plugins/Runtime/NetworkPredictionInsights/Source/NetworkPredictionInsights/Private/UI/SNPSimFrameContents.cpp
+++ b/Engine/Plugins/Runtime/NetworkPredictionInsights/Source/NetworkPredictionInsights/Private/UI/SNPSimFrameContents.cpp
@@ -2,6 +2,9 @@
 
 #include "SNPSimFrameContents.h"
 #include "Widgets/Text/SRichTextBlock.h"
+#include "Widgets/Input/SEditableText.h"
+#include "Widgets/SBoxPanel.h"
+#include "INetworkPredictionProvider.h"
 
 #define LOCTEXT_NAMESPACE "SNPSimFrameContents"
 
diff --git a/Engine/Plugins/Runtime/NetworkPredictionInsights/Source/NetworkPredictionInsights/Private/UI/SNPSimFrameView.cpp b/Engine/Plugins/Runtime/NetworkPredictionInsights/Source/NetworkPredictionInsights/Private/UI/SNPSimFrameView.cpp
--- a/Engine/Plugins/Runtime/NetworkPredictionInsights/Source/NetworkPredictionInsights/Private/UI/SNPSimFrameView.cpp
+++ b/Engine/Plugins/Runtime/NetworkPredictionInsights/Source/NetworkPredictionInsights/Private/UI/SNPSimFrameView.cpp
@@ -4,6 +4,11 @@
 #include "NetworkPredictionDrawHelpers.h"
 #include "SNPWindow.h"
 #include "Insights/IUnrealInsightsModule.h"
+#include "EditorStyleSet.h"
+#include "Styling/CoreStyle.h"
+#include "Modules/ModuleManager.h"
+#include "Framework/Commands/UIAction.h"
+#include "Textures/SlateIcon.h"
 
 #define LOCTEXT_NAMESPACE "NetworkPredictionInsights"
 
diff --git a/Engine/Plugins/Runtime/NetworkPredictionInsights/Source/NetworkPredictionInsights/Private/UI/SNPToolbar.cpp b/Engine/Plugins/Runtime/NetworkPredictionInsights/Source/NetworkPredictionInsights/Private/UI/SNPToolbar.cpp
--- a/Engine/Plugins/Runtime/NetworkPredictionInsights/Source/NetworkPredictionInsights/Private/UI/SNPToolbar.cpp
+++ b/Engine/Plugins/Runtime/NetworkPredictionInsights/Source/NetworkPredictionInsights/Private/UI/SNPToolbar.cpp
@@ -9,6 +9,10 @@
 #include "Framework/MultiBox/MultiBoxBuilder.h"
 #include "Widgets/SBoxPanel.h"
 #include "Framework/Commands/UICommandList.h"
+#include "Framework/MultiBox/MultiBoxDefs.h"
+#include "EditorStyleSet.h"
+#include "Framework/Application/SlateApplication.h"
+#include "Widgets/Layout/SBorder.h"
 
 #define LOCTEXT_NAMESPACE "NetworkPredictionInsightsModule"
 
