Removed the self-hosted Trace stores in plugins and Unreal Insights. They can assume that the trace store is already running and therefore need to just attempt to connect to it.

diff --git a/Engine/Plugins/Animation/GameplayInsights/Source/GameplayInsights/Private/GameplayInsightsModule.cpp b/Engine/Plugins/Animation/GameplayInsights/Source/GameplayInsights/Private/GameplayInsightsModule.cpp
--- a/Engine/Plugins/Animation/GameplayInsights/Source/GameplayInsights/Private/GameplayInsightsModule.cpp
+++ b/Engine/Plugins/Animation/GameplayInsights/Source/GameplayInsights/Private/GameplayInsightsModule.cpp
@@ -13,7 +13,6 @@
 #include "WorkspaceMenuStructureModule.h"
 #include "TraceServices/ITraceServicesModule.h"
 #include "Widgets/Docking/SDockTab.h"
-#include "Trace/StoreService.h"
 #include "Trace/StoreClient.h"
 #include "Stats/Stats.h"
 #include "ObjectPropertyTrace.h"
@@ -162,6 +161,7 @@ void FGameplayInsightsModule::StartupModule()
 			IUnrealInsightsModule& UnrealInsightsModule = FModuleManager::LoadModuleChecked<IUnrealInsightsModule>("TraceInsights");
 			if (!UnrealInsightsModule.GetStoreClient())
 			{
+#if WITH_TRACE_STORE
 				// Create the Store Service.
 				FString StoreDir = FPaths::ProjectSavedDir() / TEXT("TraceSessions");
 				UE::Trace::FStoreService::FDesc StoreServiceDesc;
@@ -177,6 +177,10 @@ void FGameplayInsightsModule::StartupModule()
 				// Connect to our newly created store and setup the insights module
 				UnrealInsightsModule.ConnectToStore(TEXT("localhost"), StoreService->GetPort());
 				UE::Trace::SendTo(TEXT("localhost"), StoreService->GetRecorderPort());
+#else
+				ensure(UnrealInsightsModule.ConnectToStore(TEXT("127.0.0.1")));
+				UE::Trace::SendTo(TEXT("127.0.0.1"));
+#endif // WITH_TRACE_STORE
 
 				UnrealInsightsModule.CreateSessionViewer(false);
 				UnrealInsightsModule.StartAnalysisForLastLiveSession();
diff --git a/Engine/Plugins/Animation/GameplayInsights/Source/GameplayInsights/Private/GameplayInsightsModule.h b/Engine/Plugins/Animation/GameplayInsights/Source/GameplayInsights/Private/GameplayInsightsModule.h
--- a/Engine/Plugins/Animation/GameplayInsights/Source/GameplayInsights/Private/GameplayInsightsModule.h
+++ b/Engine/Plugins/Animation/GameplayInsights/Source/GameplayInsights/Private/GameplayInsightsModule.h
@@ -9,6 +9,7 @@
 #include "Framework/Docking/TabManager.h"
 #include "IGameplayInsightsModule.h"
 #include "GameplayInsightsDebugViewCreator.h"
+#include "Trace/StoreService.h"
 
 struct FInsightsMajorTabExtender;
 
@@ -16,7 +17,9 @@ namespace UE
 {
 namespace Trace 
 {
+#if WITH_TRACE_STORE
 	class FStoreService;
+#endif
 	class FStoreClient;
 }
 }
@@ -58,6 +61,8 @@ private:
 	FDelegateHandle CustomDebugObjectHandle;
 #endif
 
+#if WITH_TRACE_STORE
 	TSharedPtr<UE::Trace::FStoreService> StoreService;
+#endif
 	TWeakPtr<FTabManager> WeakTimingProfilerTabManager;
 };
diff --git a/Engine/Plugins/Runtime/NetworkPredictionInsights/Source/NetworkPredictionInsights/Private/NetworkPredictionInsightsModule.cpp b/Engine/Plugins/Runtime/NetworkPredictionInsights/Source/NetworkPredictionInsights/Private/NetworkPredictionInsightsModule.cpp
--- a/Engine/Plugins/Runtime/NetworkPredictionInsights/Source/NetworkPredictionInsights/Private/NetworkPredictionInsightsModule.cpp
+++ b/Engine/Plugins/Runtime/NetworkPredictionInsights/Source/NetworkPredictionInsights/Private/NetworkPredictionInsightsModule.cpp
@@ -12,7 +12,6 @@
 #include "WorkspaceMenuStructureModule.h"
 #include "TraceServices/ITraceServicesModule.h"
 #include "Widgets/Docking/SDockTab.h"
-#include "Trace/StoreService.h"
 #include "Trace/StoreClient.h"
 #include "Stats/Stats.h"
 
@@ -101,6 +100,7 @@ void FNetworkPredictionInsightsModule::StartupModule()
 			IUnrealInsightsModule& UnrealInsightsModule = FModuleManager::LoadModuleChecked<IUnrealInsightsModule>("TraceInsights");
 			if (!UnrealInsightsModule.GetStoreClient())
 			{
+#if WITH_TRACE_STORE
 				// Create the Store Service.
 				FString StoreDir = FPaths::ProjectSavedDir() / TEXT("TraceSessions");
 				UE::Trace::FStoreService::FDesc StoreServiceDesc;
@@ -116,6 +116,10 @@ void FNetworkPredictionInsightsModule::StartupModule()
 				// Connect to our newly created store and setup the insights module
 				ensure(UnrealInsightsModule.ConnectToStore(TEXT("localhost"), StoreService->GetPort()));
 				UE::Trace::SendTo(TEXT("localhost"), StoreService->GetRecorderPort());
+#else
+				ensure(UnrealInsightsModule.ConnectToStore(TEXT("127.0.0.1")));
+				UE::Trace::SendTo(TEXT("127.0.0.1"));
+#endif // WITH_TRACE_STORE
 
 				UnrealInsightsModule.CreateSessionViewer(false);
 				UnrealInsightsModule.StartAnalysisForLastLiveSession();
diff --git a/Engine/Plugins/Runtime/NetworkPredictionInsights/Source/NetworkPredictionInsights/Private/NetworkPredictionInsightsModule.h b/Engine/Plugins/Runtime/NetworkPredictionInsights/Source/NetworkPredictionInsights/Private/NetworkPredictionInsightsModule.h
--- a/Engine/Plugins/Runtime/NetworkPredictionInsights/Source/NetworkPredictionInsights/Private/NetworkPredictionInsightsModule.h
+++ b/Engine/Plugins/Runtime/NetworkPredictionInsights/Source/NetworkPredictionInsights/Private/NetworkPredictionInsightsModule.h
@@ -7,6 +7,7 @@
 #include "NetworkPredictionTraceModule.h"
 #include "Framework/Docking/TabManager.h"
 #include "INetworkPredictionInsightsModule.h"
+#include "Trace/StoreService.h"
 
 struct FInsightsMajorTabExtender;
 
@@ -14,7 +15,9 @@ namespace UE
 {
 namespace Trace 
 {
+#if WITH_TRACE_STORE
 	class FStoreService;
+#endif
 	class FStoreClient;
 }
 }
@@ -32,7 +35,9 @@ private:
 	FDelegateHandle TickerHandle;
 	FDelegateHandle StoreServiceHandle;
 
+#if WITH_TRACE_STORE
 	TSharedPtr<UE::Trace::FStoreService> StoreService;
+#endif
 	//TWeakPtr<FTabManager> WeakTimingProfilerTabManager;
 
 	static const FName InsightsTabName;	
diff --git a/Engine/Source/Developer/TraceAnalysis/Public/Trace/StoreService.h b/Engine/Source/Developer/TraceAnalysis/Public/Trace/StoreService.h
--- a/Engine/Source/Developer/TraceAnalysis/Public/Trace/StoreService.h
+++ b/Engine/Source/Developer/TraceAnalysis/Public/Trace/StoreService.h
@@ -4,6 +4,15 @@
 
 #include "HAL/Platform.h"
 
+////////////////////////////////////////////////////////////////////////////////
+/* Indicates if the external store is supported or if an instance should be
+   self-hosted by something */
+#if PLATFORM_WINDOWS
+#	define WITH_TRACE_STORE 0
+#else
+#	define WITH_TRACE_STORE 1
+#endif
+
 namespace UE {
 namespace Trace {
 
diff --git a/Engine/Source/Developer/TraceInsights/Private/Insights/TraceInsightsModule.cpp b/Engine/Source/Developer/TraceInsights/Private/Insights/TraceInsightsModule.cpp
--- a/Engine/Source/Developer/TraceInsights/Private/Insights/TraceInsightsModule.cpp
+++ b/Engine/Source/Developer/TraceInsights/Private/Insights/TraceInsightsModule.cpp
@@ -105,10 +105,9 @@ void FTraceInsightsModule::UnregisterComponent(TSharedPtr<IInsightsComponent> Co
 
 void FTraceInsightsModule::CreateDefaultStore()
 {
+#if WITH_TRACE_STORE
 	const FString StoreDir = FPaths::ProjectSavedDir() / TEXT("TraceSessions");
 
-	FInsightsManager::Get()->SetStoreDir(StoreDir);
-
 	// Create the Store Service.
 	UE::Trace::FStoreService::FDesc StoreServiceDesc;
 	StoreServiceDesc.StoreDir = *StoreDir;
@@ -120,6 +119,9 @@ void FTraceInsightsModule::CreateDefaultStore()
 	{
 		ConnectToStore(TEXT("127.0.0.1"), StoreService->GetPort());
 	}
+#else
+	ConnectToStore(TEXT("127.0.0.1"));
+#endif // WITH_TRACE_STORE
 }
 
 ////////////////////////////////////////////////////////////////////////////////////////////////////
diff --git a/Engine/Source/Developer/TraceInsights/Private/Insights/TraceInsightsModule.h b/Engine/Source/Developer/TraceInsights/Private/Insights/TraceInsightsModule.h
--- a/Engine/Source/Developer/TraceInsights/Private/Insights/TraceInsightsModule.h
+++ b/Engine/Source/Developer/TraceInsights/Private/Insights/TraceInsightsModule.h
@@ -5,12 +5,15 @@
 #include "CoreMinimal.h"
 #include "Insights/IUnrealInsightsModule.h"
 #include "Framework/Docking/TabManager.h"
+#include "Trace/StoreService.h"
 
 namespace UE
 {
 namespace Trace
 {
+#if WITH_TRACE_STORE
 	class FStoreService;
+#endif
 }
 }
 
@@ -96,7 +99,9 @@ protected:
 	void OnWindowClosedEvent(const TSharedRef<SWindow>&);
 
 protected:
+#if WITH_TRACE_STORE
 	TUniquePtr<UE::Trace::FStoreService> StoreService;
+#endif
 
 	TSharedPtr<TraceServices::IAnalysisService> TraceAnalysisService;
 	TSharedPtr<TraceServices::IModuleService> TraceModuleService;
