Removed an ensure() as the store might not be available in all circumstances.

diff --git a/Engine/Plugins/Runtime/NetworkPredictionInsights/Source/NetworkPredictionInsights/Private/NetworkPredictionInsightsModule.cpp b/Engine/Plugins/Runtime/NetworkPredictionInsights/Source/NetworkPredictionInsights/Private/NetworkPredictionInsightsModule.cpp
--- a/Engine/Plugins/Runtime/NetworkPredictionInsights/Source/NetworkPredictionInsights/Private/NetworkPredictionInsightsModule.cpp
+++ b/Engine/Plugins/Runtime/NetworkPredictionInsights/Source/NetworkPredictionInsights/Private/NetworkPredictionInsightsModule.cpp
@@ -119,7 +119,7 @@ void FNetworkPredictionInsightsModule::StartupModule()
 				UE::Trace::SendTo(TEXT("localhost"), StoreService->GetRecorderPort());
 #else
 				FTraceAuxiliary::LaunchUnrealTrace();
-				ensure(UnrealInsightsModule.ConnectToStore(TEXT("127.0.0.1")));
+				UnrealInsightsModule.ConnectToStore(TEXT("127.0.0.1"));
 				UE::Trace::SendTo(TEXT("127.0.0.1"));
 #endif // WITH_TRACE_STORE
 
