- Removed explicit launching of the trace server in a couple of plugins. It was added temporarily to gather more information from CIS.

#rnx

diff --git a/Engine/Plugins/Animation/GameplayInsights/Source/GameplayInsights/Private/GameplayInsightsModule.cpp b/Engine/Plugins/Animation/GameplayInsights/Source/GameplayInsights/Private/GameplayInsightsModule.cpp
--- a/Engine/Plugins/Animation/GameplayInsights/Source/GameplayInsights/Private/GameplayInsightsModule.cpp
+++ b/Engine/Plugins/Animation/GameplayInsights/Source/GameplayInsights/Private/GameplayInsightsModule.cpp
@@ -5,7 +5,6 @@
 #include "Insights/ITimingViewExtender.h"
 #include "Containers/Ticker.h"
 #include "Modules/ModuleManager.h"
-#include "ProfilingDebugging/TraceAuxiliary.h"
 #include "Framework/Docking/LayoutExtender.h"
 #include "SAnimGraphSchematicView.h"
 #include "Insights/IUnrealInsightsModule.h"
@@ -185,7 +184,6 @@ void FGameplayInsightsModule::StartupModule()
 				UnrealInsightsModule.ConnectToStore(TEXT("localhost"), StoreService->GetPort());
 				UE::Trace::SendTo(TEXT("localhost"), StoreService->GetRecorderPort());
 #else
-				FTraceAuxiliary::LaunchUnrealTrace();
 				UnrealInsightsModule.ConnectToStore(TEXT("127.0.0.1"));
 				UE::Trace::SendTo(TEXT("127.0.0.1"));
 #endif // WITH_TRACE_STORE
diff --git a/Engine/Plugins/Runtime/NetworkPredictionInsights/Source/NetworkPredictionInsights/Private/NetworkPredictionInsightsModule.cpp b/Engine/Plugins/Runtime/NetworkPredictionInsights/Source/NetworkPredictionInsights/Private/NetworkPredictionInsightsModule.cpp
--- a/Engine/Plugins/Runtime/NetworkPredictionInsights/Source/NetworkPredictionInsights/Private/NetworkPredictionInsightsModule.cpp
+++ b/Engine/Plugins/Runtime/NetworkPredictionInsights/Source/NetworkPredictionInsights/Private/NetworkPredictionInsightsModule.cpp
@@ -5,7 +5,6 @@
 #include "Insights/ITimingViewExtender.h"
 #include "Containers/Ticker.h"
 #include "Modules/ModuleManager.h"
-#include "ProfilingDebugging/TraceAuxiliary.h"
 #include "Framework/Docking/LayoutExtender.h"
 #include "Insights/IUnrealInsightsModule.h"
 #include "HAL/PlatformApplicationMisc.h"
@@ -118,7 +117,6 @@ void FNetworkPredictionInsightsModule::StartupModule()
 				ensure(UnrealInsightsModule.ConnectToStore(TEXT("localhost"), StoreService->GetPort()));
 				UE::Trace::SendTo(TEXT("localhost"), StoreService->GetRecorderPort());
 #else
-				FTraceAuxiliary::LaunchUnrealTrace();
 				UnrealInsightsModule.ConnectToStore(TEXT("127.0.0.1"));
 				UE::Trace::SendTo(TEXT("127.0.0.1"));
 #endif // WITH_TRACE_STORE
diff --git a/Engine/Source/Runtime/Core/Private/ProfilingDebugging/TraceAuxiliary.cpp b/Engine/Source/Runtime/Core/Private/ProfilingDebugging/TraceAuxiliary.cpp
--- a/Engine/Source/Runtime/Core/Private/ProfilingDebugging/TraceAuxiliary.cpp
+++ b/Engine/Source/Runtime/Core/Private/ProfilingDebugging/TraceAuxiliary.cpp
@@ -702,11 +702,3 @@ void FTraceAuxiliary::TryAutoConnect()
 	#endif
 #endif
 }
-
-////////////////////////////////////////////////////////////////////////////////
-void FTraceAuxiliary::LaunchUnrealTrace()
-{
-#if WITH_UNREAL_TRACE_LAUNCH
-	LaunchUnrealTraceInternal();
-#endif
-}
diff --git a/Engine/Source/Runtime/Core/Public/ProfilingDebugging/TraceAuxiliary.h b/Engine/Source/Runtime/Core/Public/ProfilingDebugging/TraceAuxiliary.h
--- a/Engine/Source/Runtime/Core/Public/ProfilingDebugging/TraceAuxiliary.h
+++ b/Engine/Source/Runtime/Core/Public/ProfilingDebugging/TraceAuxiliary.h
@@ -13,5 +13,4 @@ public:
 	static void Shutdown();
 	static void TryAutoConnect();
 	static void EnableChannels();
-	static void LaunchUnrealTrace();
 };
