[Insights] Fix for GameplayInsights/NetworkInsights connection code.

diff --git a/Engine/Plugins/Animation/GameplayInsights/Source/GameplayInsights/Private/GameplayInsightsModule.cpp b/Engine/Plugins/Animation/GameplayInsights/Source/GameplayInsights/Private/GameplayInsightsModule.cpp
--- a/Engine/Plugins/Animation/GameplayInsights/Source/GameplayInsights/Private/GameplayInsightsModule.cpp
+++ b/Engine/Plugins/Animation/GameplayInsights/Source/GameplayInsights/Private/GameplayInsightsModule.cpp
@@ -29,6 +29,7 @@
 #include "SubobjectEditorMenuContext.h"
 #include "GameplayInsightsStyle.h"
 #include "SSubobjectInstanceEditor.h"
+#include "ProfilingDebugging/TraceAuxiliary.h"
 
 #endif
 
@@ -191,7 +192,10 @@ void FGameplayInsightsModule::StartupModule()
 				UE::Trace::SendTo(TEXT("localhost"), StoreService->GetRecorderPort());
 #else
 				UnrealInsightsModule.ConnectToStore(TEXT("127.0.0.1"));
-				UE::Trace::SendTo(TEXT("127.0.0.1"));
+				const bool bConnected = FTraceAuxiliary::Start(
+					FTraceAuxiliary::EConnectionType::Network,
+					TEXT("127.0.0.1"),
+					nullptr);
 #endif // WITH_TRACE_STORE
 
 				UnrealInsightsModule.CreateSessionViewer(false);
diff --git a/Engine/Plugins/Runtime/NetworkPredictionInsights/Source/NetworkPredictionInsights/Private/NetworkPredictionInsightsModule.cpp b/Engine/Plugins/Runtime/NetworkPredictionInsights/Source/NetworkPredictionInsights/Private/NetworkPredictionInsightsModule.cpp
--- a/Engine/Plugins/Runtime/NetworkPredictionInsights/Source/NetworkPredictionInsights/Private/NetworkPredictionInsightsModule.cpp
+++ b/Engine/Plugins/Runtime/NetworkPredictionInsights/Source/NetworkPredictionInsights/Private/NetworkPredictionInsightsModule.cpp
@@ -21,6 +21,7 @@
 
 #if WITH_ENGINE
 #include "Engine/Engine.h"
+#include "ProfilingDebugging/TraceAuxiliary.h"
 #endif
 
 
@@ -118,7 +119,10 @@ void FNetworkPredictionInsightsModule::StartupModule()
 				UE::Trace::SendTo(TEXT("localhost"), StoreService->GetRecorderPort());
 #else
 				UnrealInsightsModule.ConnectToStore(TEXT("127.0.0.1"));
-				UE::Trace::SendTo(TEXT("127.0.0.1"));
+				const bool bConnected = FTraceAuxiliary::Start(
+					FTraceAuxiliary::EConnectionType::Network,
+					TEXT("127.0.0.1"),
+					nullptr);
 #endif // WITH_TRACE_STORE
 
 				UnrealInsightsModule.CreateSessionViewer(false);
