[Insights] Added more memory scopes to GameplayInsights and NetworkPredictionInsights to capture more of untagged memory allocs.

diff --git a/Engine/Plugins/Animation/GameplayInsights/Source/GameplayInsights/Private/GameplayInsightsModule.cpp b/Engine/Plugins/Animation/GameplayInsights/Source/GameplayInsights/Private/GameplayInsightsModule.cpp
--- a/Engine/Plugins/Animation/GameplayInsights/Source/GameplayInsights/Private/GameplayInsightsModule.cpp
+++ b/Engine/Plugins/Animation/GameplayInsights/Source/GameplayInsights/Private/GameplayInsightsModule.cpp
@@ -8,6 +8,7 @@
 #include "Framework/Docking/LayoutExtender.h"
 #include "SAnimGraphSchematicView.h"
 #include "Insights/IUnrealInsightsModule.h"
+#include "HAL/LowLevelMemTracker.h"
 #include "HAL/PlatformApplicationMisc.h"
 #include "WorkspaceMenuStructure.h"
 #include "WorkspaceMenuStructureModule.h"
@@ -44,6 +45,8 @@ const FName GameplayInsightsTabs::DocumentTab("DocumentTab");
 
 void FGameplayInsightsModule::StartupModule()
 {
+	LLM_SCOPE_BYNAME(TEXT("Insights/GameplayInsights"));
+
 	IModularFeatures::Get().RegisterModularFeature(TraceServices::ModuleFeatureName, &GameplayTraceModule);
 	IModularFeatures::Get().RegisterModularFeature(Insights::TimingViewExtenderFeatureName, &GameplayTimingViewExtender);
 
@@ -173,6 +176,7 @@ void FGameplayInsightsModule::StartupModule()
 		// Create store and start analysis session - should only be done after engine has init and all plugins are loaded
 		FCoreDelegates::OnFEngineLoopInitComplete.AddLambda([this]
 		{
+			LLM_SCOPE_BYNAME(TEXT("Insights/GameplayInsights"));
 			IUnrealInsightsModule& UnrealInsightsModule = FModuleManager::LoadModuleChecked<IUnrealInsightsModule>("TraceInsights");
 			if (!UnrealInsightsModule.GetStoreClient())
 			{
@@ -223,6 +227,8 @@ void FGameplayInsightsModule::StartupModule()
 
 void FGameplayInsightsModule::ShutdownModule()
 {
+	LLM_SCOPE_BYNAME(TEXT("Insights/GameplayInsights"));
+
 #if OBJECT_PROPERTY_TRACE_ENABLED
 	FObjectPropertyTrace::Destroy();
 #endif
diff --git a/Engine/Plugins/Runtime/NetworkPredictionInsights/Source/NetworkPredictionInsights/Private/NetworkPredictionInsightsModule.cpp b/Engine/Plugins/Runtime/NetworkPredictionInsights/Source/NetworkPredictionInsights/Private/NetworkPredictionInsightsModule.cpp
--- a/Engine/Plugins/Runtime/NetworkPredictionInsights/Source/NetworkPredictionInsights/Private/NetworkPredictionInsightsModule.cpp
+++ b/Engine/Plugins/Runtime/NetworkPredictionInsights/Source/NetworkPredictionInsights/Private/NetworkPredictionInsightsModule.cpp
@@ -7,6 +7,7 @@
 #include "Modules/ModuleManager.h"
 #include "Framework/Docking/LayoutExtender.h"
 #include "Insights/IUnrealInsightsModule.h"
+#include "HAL/LowLevelMemTracker.h"
 #include "HAL/PlatformApplicationMisc.h"
 #include "WorkspaceMenuStructure.h"
 #include "WorkspaceMenuStructureModule.h"
@@ -33,6 +34,8 @@ const FName FNetworkPredictionInsightsModule::InsightsTabName("NetworkPrediction
 
 void FNetworkPredictionInsightsModule::StartupModule()
 {
+	LLM_SCOPE_BYNAME(TEXT("Insights/NetworkPredictionInsights"));
+
 	IModularFeatures::Get().RegisterModularFeature(TraceServices::ModuleFeatureName, &NetworkPredictionTraceModule);
 
 	FNetworkPredictionInsightsManager::Initialize();
@@ -75,6 +78,8 @@ void FNetworkPredictionInsightsModule::StartupModule()
 	FTabSpawnerEntry& TabSpawnerEntry = FGlobalTabmanager::Get()->RegisterNomadTabSpawner(FNetworkPredictionInsightsModule::InsightsTabName,
 		FOnSpawnTab::CreateLambda([](const FSpawnTabArgs& Args)
 	{
+		LLM_SCOPE_BYNAME(TEXT("Insights/NetworkPredictionInsights"));
+
 		const TSharedRef<SDockTab> DockTab = SNew(SDockTab)
 			.TabRole(ETabRole::NomadTab);
 
@@ -98,6 +103,7 @@ void FNetworkPredictionInsightsModule::StartupModule()
 		// This is temp until a more formal local server is done by the insights system.
 		StoreServiceHandle = FCoreDelegates::OnFEngineLoopInitComplete.AddLambda([this]
 		{
+			LLM_SCOPE_BYNAME(TEXT("Insights/NetworkPredictionInsights"));
 			IUnrealInsightsModule& UnrealInsightsModule = FModuleManager::LoadModuleChecked<IUnrealInsightsModule>("TraceInsights");
 			if (!UnrealInsightsModule.GetStoreClient())
 			{
@@ -137,6 +143,8 @@ void FNetworkPredictionInsightsModule::StartupModule()
 
 void FNetworkPredictionInsightsModule::ShutdownModule()
 {
+	LLM_SCOPE_BYNAME(TEXT("Insights/NetworkPredictionInsights"));
+
 	if (StoreServiceHandle.IsValid())
 	{
 		FCoreDelegates::OnFEngineLoopInitComplete.Remove(StoreServiceHandle);
