NPP: Reviving Network Prediction Insights tool.  This change fixes an engine init-time race condition between NetworkPredictionInsights and other Insights-using modules. Now NetworkPredictionInsights will make sure a trace is running even if it wasn't the one to create the store client.

Also changed the command line parameter that gated the start of the tracing to match the existing -trace=np / -trace=NetworkPrediction params that NetworkPrediction plugin uses.

diff --git a/Engine/Plugins/Runtime/NetworkPredictionInsights/Source/NetworkPredictionInsights/Private/NetworkPredictionInsightsModule.cpp b/Engine/Plugins/Runtime/NetworkPredictionInsights/Source/NetworkPredictionInsights/Private/NetworkPredictionInsightsModule.cpp
--- a/Engine/Plugins/Runtime/NetworkPredictionInsights/Source/NetworkPredictionInsights/Private/NetworkPredictionInsightsModule.cpp
+++ b/Engine/Plugins/Runtime/NetworkPredictionInsights/Source/NetworkPredictionInsights/Private/NetworkPredictionInsightsModule.cpp
@@ -15,6 +15,7 @@
 #include "Widgets/Docking/SDockTab.h"
 #include "Trace/StoreClient.h"
 #include "Stats/Stats.h"
+#include "String/ParseTokens.h"
 
 #include "UI/SNPWindow.h"
 #include "UI/NetworkPredictionInsightsManager.h"
@@ -41,6 +42,17 @@ void FNetworkPredictionInsightsModule::StartupModule()
 	FNetworkPredictionInsightsManager::Initialize();
 	IUnrealInsightsModule& UnrealInsightsModule = FModuleManager::LoadModuleChecked<IUnrealInsightsModule>("TraceInsights");
 
+	bool bShouldStartNetworkTrace = false;
+
+	FString EnabledChannels;
+	FParse::Value(FCommandLine::Get(), TEXT("-trace="), EnabledChannels, false);
+	UE::String::ParseTokens(EnabledChannels, TEXT(","), [&bShouldStartNetworkTrace](FStringView Token) {
+		if (Token.Compare(TEXT("NetworkPrediction"), ESearchCase::IgnoreCase) == 0 || Token.Compare(TEXT("NP"), ESearchCase::IgnoreCase) == 0)
+		{
+			bShouldStartNetworkTrace = true;
+		}
+		});
+
 	// Auto spawn the Network Prediction Insights tab if we detect NP data
 	// Only do this in Standalone UnrealInsights.exe. In Editor the user will select the NPI window manually.
 	// (There is currently not way to extend the high level layout of unreal insights, e.g, the layout created in FTraceInsightsModule::CreateSessionBrowser)
@@ -76,10 +88,15 @@ void FNetworkPredictionInsightsModule::StartupModule()
 
 	// Actually register our tab spawner
 	FTabSpawnerEntry& TabSpawnerEntry = FGlobalTabmanager::Get()->RegisterNomadTabSpawner(FNetworkPredictionInsightsModule::InsightsTabName,
-		FOnSpawnTab::CreateLambda([](const FSpawnTabArgs& Args)
+		FOnSpawnTab::CreateLambda([bShouldStartNetworkTrace](const FSpawnTabArgs& Args)
 	{
 		LLM_SCOPE_BYNAME(TEXT("Insights/NetworkPredictionInsights"));
 
+		if (bShouldStartNetworkTrace)
+		{
+			StartNetworkTrace();
+		}
+
 		const TSharedRef<SDockTab> DockTab = SNew(SDockTab)
 			.TabRole(ETabRole::NomadTab);
 
@@ -97,16 +114,16 @@ void FNetworkPredictionInsightsModule::StartupModule()
 	// -------------------------------------------------------------
 
 #if WITH_EDITOR
-	if (!IsRunningCommandlet())
+	if (!IsRunningCommandlet() && bShouldStartNetworkTrace)
 	{
 		// Conditionally create local store service after engine init (if someone doesn't beat us to it).
 		// This is temp until a more formal local server is done by the insights system.
-		StoreServiceHandle = FCoreDelegates::OnFEngineLoopInitComplete.AddLambda([this]
+		StoreServiceHandle = FCoreDelegates::OnFEngineLoopInitComplete.AddLambda([this, bShouldStartNetworkTrace]
 		{
 			LLM_SCOPE_BYNAME(TEXT("Insights/NetworkPredictionInsights"));
 			IUnrealInsightsModule& UnrealInsightsModule = FModuleManager::LoadModuleChecked<IUnrealInsightsModule>("TraceInsights");
-			bool bNetworkPredictionTrace = FParse::Param(FCommandLine::Get(), TEXT("NPTrace"));
-			if (!UnrealInsightsModule.GetStoreClient() && bNetworkPredictionTrace)
+
+			if (!UnrealInsightsModule.GetStoreClient())
 			{
 #if WITH_TRACE_STORE
 				UE_LOG(LogCore, Display, TEXT("NetworkPredictionInsights module auto-connecting to internal trace server..."));
@@ -128,18 +145,13 @@ void FNetworkPredictionInsightsModule::StartupModule()
 #else
 				UE_LOG(LogCore, Display, TEXT("NetworkPredictionInsights module auto-connecting to local trace server..."));
 				UnrealInsightsModule.ConnectToStore(TEXT("127.0.0.1"));
-				const bool bConnected = FTraceAuxiliary::Start(
-					FTraceAuxiliary::EConnectionType::Network,
-					TEXT("127.0.0.1"),
-					nullptr);
 #endif // WITH_TRACE_STORE
 
 				UnrealInsightsModule.CreateSessionViewer(false);
-				UnrealInsightsModule.StartAnalysisForLastLiveSession();
 			}
 		});
 	}
-#endif
+#endif // WITH_EDITOR
 }
 
 void FNetworkPredictionInsightsModule::ShutdownModule()
@@ -156,6 +168,17 @@ void FNetworkPredictionInsightsModule::ShutdownModule()
 	IModularFeatures::Get().UnregisterModularFeature(TraceServices::ModuleFeatureName, &NetworkPredictionTraceModule);
 }
 
+void FNetworkPredictionInsightsModule::StartNetworkTrace()
+{
+	const bool bConnected = FTraceAuxiliary::Start(
+		FTraceAuxiliary::EConnectionType::Network,
+		TEXT("127.0.0.1"),
+		nullptr);
+
+	IUnrealInsightsModule& UnrealInsightsModule = FModuleManager::LoadModuleChecked<IUnrealInsightsModule>("TraceInsights");
+	UnrealInsightsModule.StartAnalysisForLastLiveSession();
+}
+
 IMPLEMENT_MODULE(FNetworkPredictionInsightsModule, NetworkPredictionInsights);
 
 #undef LOCTEXT_NAMESPACE
diff --git a/Engine/Plugins/Runtime/NetworkPredictionInsights/Source/NetworkPredictionInsights/Private/NetworkPredictionInsightsModule.h b/Engine/Plugins/Runtime/NetworkPredictionInsights/Source/NetworkPredictionInsights/Private/NetworkPredictionInsightsModule.h
--- a/Engine/Plugins/Runtime/NetworkPredictionInsights/Source/NetworkPredictionInsights/Private/NetworkPredictionInsightsModule.h
+++ b/Engine/Plugins/Runtime/NetworkPredictionInsights/Source/NetworkPredictionInsights/Private/NetworkPredictionInsightsModule.h
@@ -30,6 +30,8 @@ public:
 	virtual void ShutdownModule() override;
 
 private:
+	static void StartNetworkTrace();
+
 	FNetworkPredictionTraceModule NetworkPredictionTraceModule;
 
 	FTSTicker::FDelegateHandle TickerHandle;
