Fixing Build Tools failure in NetworkPredictionInsights, caused by some recently-shuffled code that moved outside of a WITH_EDITOR block

#rnx

diff --git a/Engine/Plugins/Runtime/NetworkPredictionInsights/Source/NetworkPredictionInsights/Private/NetworkPredictionInsightsModule.cpp b/Engine/Plugins/Runtime/NetworkPredictionInsights/Source/NetworkPredictionInsights/Private/NetworkPredictionInsightsModule.cpp
--- a/Engine/Plugins/Runtime/NetworkPredictionInsights/Source/NetworkPredictionInsights/Private/NetworkPredictionInsightsModule.cpp
+++ b/Engine/Plugins/Runtime/NetworkPredictionInsights/Source/NetworkPredictionInsights/Private/NetworkPredictionInsightsModule.cpp
@@ -170,6 +170,7 @@ void FNetworkPredictionInsightsModule::ShutdownModule()
 
 void FNetworkPredictionInsightsModule::StartNetworkTrace()
 {
+#if WITH_EDITOR
 	const bool bConnected = FTraceAuxiliary::Start(
 		FTraceAuxiliary::EConnectionType::Network,
 		TEXT("127.0.0.1"),
@@ -177,6 +178,7 @@ void FNetworkPredictionInsightsModule::StartNetworkTrace()
 
 	IUnrealInsightsModule& UnrealInsightsModule = FModuleManager::LoadModuleChecked<IUnrealInsightsModule>("TraceInsights");
 	UnrealInsightsModule.StartAnalysisForLastLiveSession();
+#endif
 }
 
 IMPLEMENT_MODULE(FNetworkPredictionInsightsModule, NetworkPredictionInsights);
