Removed unnecessary variable from capture clause. No functional changes.

#rnx

diff --git a/Engine/Plugins/Runtime/NetworkPredictionInsights/Source/NetworkPredictionInsights/Private/NetworkPredictionInsightsModule.cpp b/Engine/Plugins/Runtime/NetworkPredictionInsights/Source/NetworkPredictionInsights/Private/NetworkPredictionInsightsModule.cpp
--- a/Engine/Plugins/Runtime/NetworkPredictionInsights/Source/NetworkPredictionInsights/Private/NetworkPredictionInsightsModule.cpp
+++ b/Engine/Plugins/Runtime/NetworkPredictionInsights/Source/NetworkPredictionInsights/Private/NetworkPredictionInsightsModule.cpp
@@ -118,7 +118,7 @@ void FNetworkPredictionInsightsModule::StartupModule()
 	{
 		// Conditionally create local store service after engine init (if someone doesn't beat us to it).
 		// This is temp until a more formal local server is done by the insights system.
-		StoreServiceHandle = FCoreDelegates::OnFEngineLoopInitComplete.AddLambda([this, bShouldStartNetworkTrace]
+		StoreServiceHandle = FCoreDelegates::OnFEngineLoopInitComplete.AddLambda([this]
 		{
 			LLM_SCOPE_BYNAME(TEXT("Insights/NetworkPredictionInsights"));
 			IUnrealInsightsModule& UnrealInsightsModule = FModuleManager::LoadModuleChecked<IUnrealInsightsModule>("TraceInsights");
