Fixed disableunity/nopch compile error for NetworkPredictionInsightsModule (missing Misc/CommandLine.h include). Also ordered the includes alphabetically.

diff --git a/Engine/Plugins/Runtime/NetworkPredictionInsights/Source/NetworkPredictionInsights/Private/NetworkPredictionInsightsModule.cpp b/Engine/Plugins/Runtime/NetworkPredictionInsights/Source/NetworkPredictionInsights/Private/NetworkPredictionInsightsModule.cpp
--- a/Engine/Plugins/Runtime/NetworkPredictionInsights/Source/NetworkPredictionInsights/Private/NetworkPredictionInsightsModule.cpp
+++ b/Engine/Plugins/Runtime/NetworkPredictionInsights/Source/NetworkPredictionInsights/Private/NetworkPredictionInsightsModule.cpp
@@ -1,25 +1,27 @@
 // Copyright Epic Games, Inc. All Rights Reserved.
 
 #include "NetworkPredictionInsightsModule.h"
-#include "Features/IModularFeatures.h"
-#include "Insights/ITimingViewExtender.h"
+
 #include "Containers/Ticker.h"
-#include "Modules/ModuleManager.h"
+#include "Features/IModularFeatures.h"
 #include "Framework/Docking/LayoutExtender.h"
-#include "Insights/IUnrealInsightsModule.h"
 #include "HAL/LowLevelMemTracker.h"
 #include "HAL/PlatformApplicationMisc.h"
-#include "WorkspaceMenuStructure.h"
-#include "WorkspaceMenuStructureModule.h"
-#include "TraceServices/ITraceServicesModule.h"
-#include "Widgets/Docking/SDockTab.h"
-#include "Trace/StoreClient.h"
+#include "Insights/ITimingViewExtender.h"
+#include "Insights/IUnrealInsightsModule.h"
+#include "Misc/CommandLine.h"
+#include "Modules/ModuleManager.h"
 #include "Stats/Stats.h"
 #include "String/ParseTokens.h"
+#include "Styling/AppStyle.h"
+#include "Trace/StoreClient.h"
+#include "TraceServices/ITraceServicesModule.h"
+#include "Widgets/Docking/SDockTab.h"
+#include "WorkspaceMenuStructure.h"
+#include "WorkspaceMenuStructureModule.h"
 
-#include "UI/SNPWindow.h"
 #include "UI/NetworkPredictionInsightsManager.h"
-#include "Styling/AppStyle.h"
+#include "UI/SNPWindow.h"
 
 #if WITH_ENGINE
 #include "Engine/Engine.h"
