Make sure network prediction trace channel is always enabled as we start a trace

diff --git a/Engine/Plugins/Runtime/NetworkPredictionInsights/Source/NetworkPredictionInsights/Private/NetworkPredictionInsightsModule.cpp b/Engine/Plugins/Runtime/NetworkPredictionInsights/Source/NetworkPredictionInsights/Private/NetworkPredictionInsightsModule.cpp
--- a/Engine/Plugins/Runtime/NetworkPredictionInsights/Source/NetworkPredictionInsights/Private/NetworkPredictionInsightsModule.cpp
+++ b/Engine/Plugins/Runtime/NetworkPredictionInsights/Source/NetworkPredictionInsights/Private/NetworkPredictionInsightsModule.cpp
@@ -173,6 +173,8 @@ void FNetworkPredictionInsightsModule::ShutdownModule()
 void FNetworkPredictionInsightsModule::StartNetworkTrace()
 {
 #if WITH_EDITOR
+	UE::Trace::ToggleChannel(TEXT("NetworkPredictionChannel"), true);
+
 	const bool bConnected = FTraceAuxiliary::Start(
 		FTraceAuxiliary::EConnectionType::Network,
 		TEXT("127.0.0.1"),
