Network Prediction plugin: fixed incorrect enum description. No functional changes.

#rnx

diff --git a/Engine/Plugins/Runtime/NetworkPredictionInsights/Source/NetworkPredictionInsights/Public/INetworkPredictionProvider.h b/Engine/Plugins/Runtime/NetworkPredictionInsights/Source/NetworkPredictionInsights/Public/INetworkPredictionProvider.h
--- a/Engine/Plugins/Runtime/NetworkPredictionInsights/Source/NetworkPredictionInsights/Public/INetworkPredictionProvider.h
+++ b/Engine/Plugins/Runtime/NetworkPredictionInsights/Source/NetworkPredictionInsights/Public/INetworkPredictionProvider.h
@@ -36,7 +36,7 @@ enum class ENP_TickingPolicy: uint8
 	Fixed = 1 << 1
 };
 
-// Must be kept in sync with ENetworkPredictionTickingPolicy
+// Must be kept in sync with ENetworkLOD
 enum class ENP_NetworkLOD: uint8
 {
 	Interpolated,
