Remove inactive trace store code

diff --git a/Engine/Plugins/Animation/GameplayInsights/Source/GameplayInsights/Private/GameplayInsightsModule.cpp b/Engine/Plugins/Animation/GameplayInsights/Source/GameplayInsights/Private/GameplayInsightsModule.cpp
--- a/Engine/Plugins/Animation/GameplayInsights/Source/GameplayInsights/Private/GameplayInsightsModule.cpp
+++ b/Engine/Plugins/Animation/GameplayInsights/Source/GameplayInsights/Private/GameplayInsightsModule.cpp
@@ -193,28 +193,8 @@ void FGameplayInsightsModule::StartupModule()
 			IUnrealInsightsModule& UnrealInsightsModule = FModuleManager::LoadModuleChecked<IUnrealInsightsModule>("TraceInsights");
 			if (!UnrealInsightsModule.GetStoreClient())
 			{
-#if WITH_TRACE_STORE
-				UE_LOG(LogCore, Display, TEXT("GameplayInsights module auto-connecting to internal trace server..."));
-				// Create the Store Service.
-				FString StoreDir = FPaths::ProjectSavedDir() / TEXT("TraceSessions");
-				UE::Trace::FStoreService::FDesc StoreServiceDesc;
-				StoreServiceDesc.StoreDir = *StoreDir;
-				StoreServiceDesc.RecorderPort = 0; // Let system decide port
-				StoreServiceDesc.ThreadCount = 2;
-				StoreService = TSharedPtr<UE::Trace::FStoreService>(UE::Trace::FStoreService::Create(StoreServiceDesc));
-
-				FCoreDelegates::OnPreExit.AddLambda([this]() {
-					StoreService.Reset();
-				});
-
-				// Connect to our newly created store and setup the insights module
-				UnrealInsightsModule.ConnectToStore(TEXT("localhost"), StoreService->GetPort());
-				UE::Trace::SendTo(TEXT("localhost"), StoreService->GetRecorderPort());
-#else
 				UE_LOG(LogCore, Display, TEXT("GameplayInsights module auto-connecting to local trace server..."));
 				UnrealInsightsModule.ConnectToStore(TEXT("127.0.0.1"));
-#endif // WITH_TRACE_STORE
-
 				UnrealInsightsModule.CreateSessionViewer(false);
 			}
 		});
diff --git a/Engine/Plugins/Animation/GameplayInsights/Source/GameplayInsights/Private/GameplayInsightsModule.h b/Engine/Plugins/Animation/GameplayInsights/Source/GameplayInsights/Private/GameplayInsightsModule.h
--- a/Engine/Plugins/Animation/GameplayInsights/Source/GameplayInsights/Private/GameplayInsightsModule.h
+++ b/Engine/Plugins/Animation/GameplayInsights/Source/GameplayInsights/Private/GameplayInsightsModule.h
@@ -16,9 +16,6 @@ namespace UE
 {
 namespace Trace 
 {
-#if WITH_TRACE_STORE
-	class FStoreService;
-#endif
 	class FStoreClient;
 }
 }
@@ -56,9 +53,6 @@ private:
 	FDelegateHandle CustomDebugObjectHandle;
 #endif
 
-#if WITH_TRACE_STORE
-	TSharedPtr<UE::Trace::FStoreService> StoreService;
-#endif
 	TWeakPtr<FTabManager> WeakTimingProfilerTabManager;
 	
 	bool bTraceStarted = false;
diff --git a/Engine/Plugins/Runtime/NetworkPredictionInsights/Source/NetworkPredictionInsights/Private/NetworkPredictionInsightsModule.cpp b/Engine/Plugins/Runtime/NetworkPredictionInsights/Source/NetworkPredictionInsights/Private/NetworkPredictionInsightsModule.cpp
--- a/Engine/Plugins/Runtime/NetworkPredictionInsights/Source/NetworkPredictionInsights/Private/NetworkPredictionInsightsModule.cpp
+++ b/Engine/Plugins/Runtime/NetworkPredictionInsights/Source/NetworkPredictionInsights/Private/NetworkPredictionInsightsModule.cpp
@@ -120,28 +120,8 @@ void FNetworkPredictionInsightsModule::StartupModule()
 
 			if (!UnrealInsightsModule.GetStoreClient())
 			{
-#if WITH_TRACE_STORE
-				UE_LOG(LogCore, Display, TEXT("NetworkPredictionInsights module auto-connecting to internal trace server..."));
-				// Create the Store Service.
-				FString StoreDir = FPaths::ProjectSavedDir() / TEXT("TraceSessions");
-				UE::Trace::FStoreService::FDesc StoreServiceDesc;
-				StoreServiceDesc.StoreDir = *StoreDir;
-				StoreServiceDesc.RecorderPort = 0; // Let system decide port
-				StoreServiceDesc.ThreadCount = 2;
-				StoreService = TSharedPtr<UE::Trace::FStoreService>(UE::Trace::FStoreService::Create(StoreServiceDesc));
-
-				FCoreDelegates::OnPreExit.AddLambda([this]() {
-					StoreService.Reset();
-				});
-
-				// Connect to our newly created store and setup the insights module
-				ensure(UnrealInsightsModule.ConnectToStore(TEXT("localhost"), StoreService->GetPort()));
-				UE::Trace::SendTo(TEXT("localhost"), StoreService->GetRecorderPort());
-#else
 				UE_LOG(LogCore, Display, TEXT("NetworkPredictionInsights module auto-connecting to local trace server..."));
 				UnrealInsightsModule.ConnectToStore(TEXT("127.0.0.1"));
-#endif // WITH_TRACE_STORE
-
 				UnrealInsightsModule.CreateSessionViewer(false);
 			}
 		});
diff --git a/Engine/Plugins/Runtime/NetworkPredictionInsights/Source/NetworkPredictionInsights/Private/NetworkPredictionInsightsModule.h b/Engine/Plugins/Runtime/NetworkPredictionInsights/Source/NetworkPredictionInsights/Private/NetworkPredictionInsightsModule.h
--- a/Engine/Plugins/Runtime/NetworkPredictionInsights/Source/NetworkPredictionInsights/Private/NetworkPredictionInsightsModule.h
+++ b/Engine/Plugins/Runtime/NetworkPredictionInsights/Source/NetworkPredictionInsights/Private/NetworkPredictionInsightsModule.h
@@ -13,9 +13,6 @@ namespace UE
 {
 namespace Trace 
 {
-#if WITH_TRACE_STORE
-	class FStoreService;
-#endif
 	class FStoreClient;
 }
 }
@@ -35,9 +32,6 @@ private:
 	FTSTicker::FDelegateHandle TickerHandle;
 	FDelegateHandle StoreServiceHandle;
 
-#if WITH_TRACE_STORE
-	TSharedPtr<UE::Trace::FStoreService> StoreService;
-#endif
 	//TWeakPtr<FTabManager> WeakTimingProfilerTabManager;
 
 	static const FName InsightsTabName;	
diff --git a/Engine/Source/Developer/TraceAnalysis/Private/Store/AsioContext.cpp b/Engine/Source/Developer/TraceAnalysis/Private/Store/AsioContext.cpp
deleted file mode 100644
diff --git a/Engine/Source/Developer/TraceAnalysis/Private/Store/AsioContext.h b/Engine/Source/Developer/TraceAnalysis/Private/Store/AsioContext.h
deleted file mode 100644
diff --git a/Engine/Source/Developer/TraceAnalysis/Private/Store/AsioFile.cpp b/Engine/Source/Developer/TraceAnalysis/Private/Store/AsioFile.cpp
deleted file mode 100644
diff --git a/Engine/Source/Developer/TraceAnalysis/Private/Store/AsioFile.h b/Engine/Source/Developer/TraceAnalysis/Private/Store/AsioFile.h
deleted file mode 100644
diff --git a/Engine/Source/Developer/TraceAnalysis/Private/Store/AsioIoable.cpp b/Engine/Source/Developer/TraceAnalysis/Private/Store/AsioIoable.cpp
deleted file mode 100644
diff --git a/Engine/Source/Developer/TraceAnalysis/Private/Store/AsioIoable.h b/Engine/Source/Developer/TraceAnalysis/Private/Store/AsioIoable.h
deleted file mode 100644
diff --git a/Engine/Source/Developer/TraceAnalysis/Private/Store/AsioRecorder.cpp b/Engine/Source/Developer/TraceAnalysis/Private/Store/AsioRecorder.cpp
deleted file mode 100644
diff --git a/Engine/Source/Developer/TraceAnalysis/Private/Store/AsioRecorder.h b/Engine/Source/Developer/TraceAnalysis/Private/Store/AsioRecorder.h
deleted file mode 100644
diff --git a/Engine/Source/Developer/TraceAnalysis/Private/Store/AsioSocket.cpp b/Engine/Source/Developer/TraceAnalysis/Private/Store/AsioSocket.cpp
deleted file mode 100644
diff --git a/Engine/Source/Developer/TraceAnalysis/Private/Store/AsioSocket.h b/Engine/Source/Developer/TraceAnalysis/Private/Store/AsioSocket.h
deleted file mode 100644
diff --git a/Engine/Source/Developer/TraceAnalysis/Private/Store/AsioStore.cpp b/Engine/Source/Developer/TraceAnalysis/Private/Store/AsioStore.cpp
deleted file mode 100644
diff --git a/Engine/Source/Developer/TraceAnalysis/Private/Store/AsioStore.h b/Engine/Source/Developer/TraceAnalysis/Private/Store/AsioStore.h
deleted file mode 100644
diff --git a/Engine/Source/Developer/TraceAnalysis/Private/Store/AsioStoreCborServer.cpp b/Engine/Source/Developer/TraceAnalysis/Private/Store/AsioStoreCborServer.cpp
deleted file mode 100644
diff --git a/Engine/Source/Developer/TraceAnalysis/Private/Store/AsioStoreCborServer.h b/Engine/Source/Developer/TraceAnalysis/Private/Store/AsioStoreCborServer.h
deleted file mode 100644
diff --git a/Engine/Source/Developer/TraceAnalysis/Private/Store/AsioTcpServer.cpp b/Engine/Source/Developer/TraceAnalysis/Private/Store/AsioTcpServer.cpp
deleted file mode 100644
diff --git a/Engine/Source/Developer/TraceAnalysis/Private/Store/AsioTcpServer.h b/Engine/Source/Developer/TraceAnalysis/Private/Store/AsioTcpServer.h
deleted file mode 100644
diff --git a/Engine/Source/Developer/TraceAnalysis/Private/Store/AsioTickable.cpp b/Engine/Source/Developer/TraceAnalysis/Private/Store/AsioTickable.cpp
deleted file mode 100644
diff --git a/Engine/Source/Developer/TraceAnalysis/Private/Store/AsioTickable.h b/Engine/Source/Developer/TraceAnalysis/Private/Store/AsioTickable.h
deleted file mode 100644
diff --git a/Engine/Source/Developer/TraceAnalysis/Private/Store/AsioTraceRelay.cpp b/Engine/Source/Developer/TraceAnalysis/Private/Store/AsioTraceRelay.cpp
deleted file mode 100644
diff --git a/Engine/Source/Developer/TraceAnalysis/Private/Store/AsioTraceRelay.h b/Engine/Source/Developer/TraceAnalysis/Private/Store/AsioTraceRelay.h
deleted file mode 100644
diff --git a/Engine/Source/Developer/TraceAnalysis/Private/Store/StoreService.cpp b/Engine/Source/Developer/TraceAnalysis/Private/Store/StoreService.cpp
deleted file mode 100644
diff --git a/Engine/Source/Developer/TraceAnalysis/Private/Store/Utils.h b/Engine/Source/Developer/TraceAnalysis/Private/Store/Utils.h
deleted file mode 100644
diff --git a/Engine/Source/Developer/TraceAnalysis/Public/Trace/StoreService.h b/Engine/Source/Developer/TraceAnalysis/Public/Trace/StoreService.h
--- a/Engine/Source/Developer/TraceAnalysis/Public/Trace/StoreService.h
+++ b/Engine/Source/Developer/TraceAnalysis/Public/Trace/StoreService.h
@@ -5,10 +5,6 @@
 #include "HAL/Platform.h"
 
 ////////////////////////////////////////////////////////////////////////////////
-/* Indicates if the external store is supported or if an instance should be
-   self-hosted by something */
-#define WITH_TRACE_STORE 0
-
 namespace UE {
 namespace Trace {
 
diff --git a/Engine/Source/Developer/TraceInsights/Private/Insights/TraceInsightsModule.cpp b/Engine/Source/Developer/TraceInsights/Private/Insights/TraceInsightsModule.cpp
--- a/Engine/Source/Developer/TraceInsights/Private/Insights/TraceInsightsModule.cpp
+++ b/Engine/Source/Developer/TraceInsights/Private/Insights/TraceInsightsModule.cpp
@@ -177,23 +177,7 @@ void FTraceInsightsModule::UnregisterComponent(TSharedPtr<IInsightsComponent> Co
 
 void FTraceInsightsModule::CreateDefaultStore()
 {
-#if WITH_TRACE_STORE
-	const FString StoreDir = FPaths::ProjectSavedDir() / TEXT("TraceSessions");
-
-	// Create the Store Service.
-	UE::Trace::FStoreService::FDesc StoreServiceDesc;
-	StoreServiceDesc.StoreDir = *StoreDir;
-	StoreServiceDesc.RecorderPort = 1981;
-	StoreServiceDesc.ThreadCount = 2;
-	StoreService.Reset(UE::Trace::FStoreService::Create(StoreServiceDesc));
-
-	if (StoreService.IsValid())
-	{
-		ConnectToStore(TEXT("127.0.0.1"), StoreService->GetPort());
-	}
-#else
 	ConnectToStore(TEXT("127.0.0.1"));
-#endif // WITH_TRACE_STORE
 }
 
 ////////////////////////////////////////////////////////////////////////////////////////////////////
diff --git a/Engine/Source/Developer/TraceInsights/Private/Insights/TraceInsightsModule.h b/Engine/Source/Developer/TraceInsights/Private/Insights/TraceInsightsModule.h
--- a/Engine/Source/Developer/TraceInsights/Private/Insights/TraceInsightsModule.h
+++ b/Engine/Source/Developer/TraceInsights/Private/Insights/TraceInsightsModule.h
@@ -7,16 +7,6 @@
 #include "Framework/Docking/TabManager.h"
 #include "Trace/StoreService.h"
 
-namespace UE
-{
-namespace Trace
-{
-#if WITH_TRACE_STORE
-	class FStoreService;
-#endif
-}
-}
-
 namespace TraceServices
 {
 	class IAnalysisService;
@@ -111,11 +101,6 @@ protected:
 
 	void HandleCodeAccessorOpenFileFailed(const FString& Filename);
 
-protected:
-#if WITH_TRACE_STORE
-	TUniquePtr<UE::Trace::FStoreService> StoreService;
-#endif
-
 	TSharedPtr<TraceServices::IAnalysisService> TraceAnalysisService;
 	TSharedPtr<TraceServices::IModuleService> TraceModuleService;
 
