Hotfixing incremental build

#rnx

diff --git a/Engine/Plugins/Experimental/Avalanche/Source/AvalancheOutliner/Private/Columns/Slate/SAvaOutlinerLabelComponent.cpp b/Engine/Plugins/Experimental/Avalanche/Source/AvalancheOutliner/Private/Columns/Slate/SAvaOutlinerLabelComponent.cpp
--- a/Engine/Plugins/Experimental/Avalanche/Source/AvalancheOutliner/Private/Columns/Slate/SAvaOutlinerLabelComponent.cpp
+++ b/Engine/Plugins/Experimental/Avalanche/Source/AvalancheOutliner/Private/Columns/Slate/SAvaOutlinerLabelComponent.cpp
@@ -5,6 +5,7 @@
 #include "Item/AvaOutlinerComponent.h"
 #include "Misc/Paths.h"
 #include "Styling/SlateStyleMacros.h"
+#include "Styling/CoreStyle.h"
 
 #define LOCTEXT_NAMESPACE "SAvaOutlinerLabelComponent"
 
diff --git a/Engine/Plugins/Runtime/NetworkPredictionInsights/Source/NetworkPredictionInsights/Private/NetworkPredictionInsightsStyle.h b/Engine/Plugins/Runtime/NetworkPredictionInsights/Source/NetworkPredictionInsights/Private/NetworkPredictionInsightsStyle.h
--- a/Engine/Plugins/Runtime/NetworkPredictionInsights/Source/NetworkPredictionInsights/Private/NetworkPredictionInsightsStyle.h
+++ b/Engine/Plugins/Runtime/NetworkPredictionInsights/Source/NetworkPredictionInsights/Private/NetworkPredictionInsightsStyle.h
@@ -4,6 +4,8 @@
 
 // HEADER_UNIT_SKIP - TODO
 
+/*
+
 #include "Styling/SlateStyle.h"
 #include "Styling/SlateStyleRegistry.h"
 #include "Styling/SlateStyleMacros.h"
@@ -75,8 +77,10 @@ public:
 	const TSharedRef<FLinearColor> SimFrameColor_Corrected_Ref;
 	const TSharedRef<FLinearColor> SimFrameColor_Predicted_Ref;
 	const FSlateColor SimFrameColor_Confirmed;
-	const FSlateColor SimFrameColor_Confirmed;
-	const FSlateColor SimFrameColor_Confirmed;
+	const FSlateColor SimFrameColor_Corrected;
+	const FSlateColor SimFrameColor_Predicted;
 
-	0.53333  0.67059  0.53333
+	//0.53333  0.67059  0.53333
 };
+
+*/
